Added support for zen partial downloads.

#Synced-CL 51662096

#Synced-CL 51662096

diff --git a/Engine/Source/Developer/IoStore/OnDemandUtilities/Private/Test.cpp b/Engine/Source/Developer/IoStore/OnDemandUtilities/Private/Test.cpp
--- a/Engine/Source/Developer/IoStore/OnDemandUtilities/Private/Test.cpp
+++ b/Engine/Source/Developer/IoStore/OnDemandUtilities/Private/Test.cpp
@@ -40,25 +40,26 @@ static void CacheTests(const FContext& Context)
 ////////////////////////////////////////////////////////////////////////////////
 static int32 TestCommandEntry(const FContext& Context)
 {
-	FStringView Only = Context.Get<FStringView>(TEXT("-Only"));
 
+	FStringView Only = Context.Get<FStringView>(TEXT("-Only"));
+#if !UE_BUILD_SHIPPING
 	if (Only.IsEmpty())
 	{
 		CommandTest();
 	}
-
+#endif
 #if IS_PROGRAM
 	if (Only.IsEmpty() || Only == TEXT("cache"))
 	{
 		CacheTests(Context);
 	}
 #endif //IS_PROGRAM
-
+#if !UE_BUILD_SHIPPING
 	if (Only.IsEmpty() || Only == TEXT("http"))
 	{
 		HttpTests(Context);
 	}
-
+#endif
 	return 0;
 }
 
diff --git a/Engine/Source/Programs/IoStoreDependencyViewer/PackageForDistribution.ps1 b/Engine/Source/Programs/IoStoreDependencyViewer/PackageForDistribution.ps1
--- a/Engine/Source/Programs/IoStoreDependencyViewer/PackageForDistribution.ps1
+++ b/Engine/Source/Programs/IoStoreDependencyViewer/PackageForDistribution.ps1
@@ -28,6 +28,8 @@ Write-Host "Reading settings from: $SettingsPath" -ForegroundColor Cyan
 $IniContent = Get-Content $SettingsPath
 $OutputDir = $null
 $Version = $null
+$IncludeZenTools = $false
+$ZenReleasePath = $null
 
 foreach ($Line in $IniContent) {
     $Line = $Line.Trim()
@@ -37,6 +39,12 @@ foreach ($Line in $IniContent) {
     if ($Line -match '^Version=(.+)$') {
         $Version = $Matches[1].Trim()
     }
+    if ($Line -match '^IncludeZenTools=(.+)$') {
+        $IncludeZenTools = $Matches[1].Trim() -eq 'true'
+    }
+    if ($Line -match '^ZenReleasePath=(.+)$') {
+        $ZenReleasePath = $Matches[1].Trim()
+    }
 }
 
 if (-not $OutputDir) {
@@ -51,10 +59,15 @@ if (-not $Version) {
 
 Write-Host "  Output Directory: $OutputDir" -ForegroundColor Green
 Write-Host "  Version: $Version" -ForegroundColor Green
+Write-Host "  Include Zen Tools: $IncludeZenTools" -ForegroundColor $(if ($IncludeZenTools) { "Green" } else { "Yellow" })
+if ($IncludeZenTools -and $ZenReleasePath) {
+    Write-Host "  Zen Release Path: $ZenReleasePath" -ForegroundColor Green
+}
 Write-Host ""
 
 # Output paths - fixed filename
-$TempDir = Join-Path $env:TEMP "IoStoreDependencyViewer_Package"
+# Use D: drive for temp to avoid C: drive space issues
+$TempDir = "D:\Temp\IoStoreDependencyViewer_Package"
 $ZipName = "IoStoreDependencyViewer.zip"
 $ZipPath = Join-Path $OutputDir $ZipName
 
@@ -86,6 +99,12 @@ New-Item -ItemType Directory -Path $BinDir -Force | Out-Null
 Copy-Item $ExePath -Destination $BinDir
 
 # Create version.txt
+$ZenToolsInfo = if ($IncludeZenTools) {
+    "`n`nExperimental zen.exe included in Engine\Binaries\Win64\ (partial download support)"
+} else {
+    ""
+}
+
 $VersionContent = @"
 IoStoreDependencyViewer
 Version: $Version
@@ -95,7 +114,7 @@ A standalone tool for analyzing and visualizing IoStore container dependencies i
 
 For documentation and usage instructions, see README.md
 
-To run: Execute Engine\Binaries\Win64\IoStoreDependencyViewer.exe
+To run: Execute Engine\Binaries\Win64\IoStoreDependencyViewer.exe$ZenToolsInfo
 "@
 
 $VersionPath = Join-Path $TempDir "version.txt"
@@ -214,6 +233,34 @@ if ($MissingModules.Count -gt 0) {
 }
 Write-Host ""
 
+# Copy experimental zen.exe if enabled
+if ($IncludeZenTools) {
+    Write-Host "Including experimental zen.exe..." -ForegroundColor Cyan
+
+    if (-not $ZenReleasePath) {
+        Write-Warning "IncludeZenTools=true but ZenReleasePath not specified in PackageSettings.ini"
+        Write-Host "  [SKIP] zen.exe" -ForegroundColor Yellow
+    } elseif (-not (Test-Path $ZenReleasePath)) {
+        Write-Warning "ZenReleasePath does not exist: $ZenReleasePath"
+        Write-Host "  [SKIP] zen.exe" -ForegroundColor Yellow
+    } else {
+        # Copy zen.exe to same directory as IoStoreDependencyViewer.exe
+        $ZenExePath = Join-Path $ZenReleasePath "zen.exe"
+        if (Test-Path $ZenExePath) {
+            Copy-Item $ZenExePath -Destination $BinDir
+            $zenSize = (Get-Item $ZenExePath).Length / 1MB
+            Write-Host "  [OK] zen.exe ($("{0:N2} MB" -f $zenSize)) - experimental build with partial download support" -ForegroundColor Green
+        } else {
+            Write-Host "  [SKIP] zen.exe not found at: $ZenExePath" -ForegroundColor Yellow
+        }
+    }
+    Write-Host ""
+} else {
+    Write-Host "Zen.exe not included (IncludeZenTools=false)" -ForegroundColor Yellow
+    Write-Host "  Users must provide their own zen.exe for cloud downloads" -ForegroundColor Yellow
+    Write-Host ""
+}
+
 # Create output directory if it doesn't exist
 if (-not (Test-Path $OutputDir)) {
     New-Item -ItemType Directory -Path $OutputDir | Out-Null
@@ -244,11 +291,13 @@ Remove-Item $TempDir -Recurse -Force
 
 # Display results
 $ZipSize = (Get-Item $ZipPath).Length / 1MB
+$ZenToolsStatus = if ($IncludeZenTools) { "Included" } else { "Not Included" }
 Write-Host ""
 Write-Host "=== Package Created Successfully ===" -ForegroundColor Green
 Write-Host "  Output: $ZipPath"
 Write-Host "  Size: $("{0:N2} MB" -f $ZipSize)"
 Write-Host "  Version: $Version"
 Write-Host "  Files: IoStoreDependencyViewer.exe + $CopiedCount DLLs + version.txt + README.md"
+Write-Host "  Zen Tools: $ZenToolsStatus" -ForegroundColor $(if ($IncludeZenTools) { "Green" } else { "Yellow" })
 Write-Host ""
 Write-Host "Package is ready for distribution!" -ForegroundColor Cyan
diff --git a/Engine/Source/Programs/IoStoreDependencyViewer/PackageSettings.ini b/Engine/Source/Programs/IoStoreDependencyViewer/PackageSettings.ini
--- a/Engine/Source/Programs/IoStoreDependencyViewer/PackageSettings.ini
+++ b/Engine/Source/Programs/IoStoreDependencyViewer/PackageSettings.ini
@@ -3,4 +3,12 @@
 OutputPath=D:\Depot\depot\InternalTools\UnrealGameSync\UnrealGameSync\bin\Tools\IoStoreDependencyViewer
 
 # Version number (will be written to version.txt inside the package)
-Version=1.0.1
+Version=1.0.3
+
+# Include experimental zen.exe with partial download support (optional)
+# Set to true to include zen.exe and OidcToken.exe in the package
+# Set to false or comment out to exclude zen tools (users must provide their own)
+IncludeZenTools=true
+
+# Path to experimental zen.exe release build (only used if IncludeZenTools=true)
+ZenReleasePath=D:\github\zen\Zen\build\windows\x64\release
diff --git a/Engine/Source/Programs/IoStoreDependencyViewer/Private/IoStoreDependencyViewer.cpp b/Engine/Source/Programs/IoStoreDependencyViewer/Private/IoStoreDependencyViewer.cpp
--- a/Engine/Source/Programs/IoStoreDependencyViewer/Private/IoStoreDependencyViewer.cpp
+++ b/Engine/Source/Programs/IoStoreDependencyViewer/Private/IoStoreDependencyViewer.cpp
@@ -90,6 +90,27 @@ INT32_MAIN_INT32_ARGC_TCHAR_ARGV()
 		ExitCode = bSuccess ? 0 : 1;
 		bRunGuiMode = false;
 	}
+	// Check for cloud partial download UCAS reading test
+	else if (FParse::Param(FCommandLine::Get(), TEXT("TestCloudPartialDownloadUCas")))
+	{
+		bool bSuccess = RunTestCloudPartialDownloadUCas();
+		ExitCode = bSuccess ? 0 : 1;
+		bRunGuiMode = false;
+	}
+	// Check for zen reader integration test
+	else if (FParse::Param(FCommandLine::Get(), TEXT("TestZenReaderIntegration")))
+	{
+		bool bSuccess = RunTestZenReaderIntegration();
+		ExitCode = bSuccess ? 0 : 1;
+		bRunGuiMode = false;
+	}
+	// Check for partial download test
+	else if (FParse::Param(FCommandLine::Get(), TEXT("TestPartialDownload")))
+	{
+		TestPartialDownload();
+		ExitCode = 0;
+		bRunGuiMode = false;
+	}
 
 	// Run GUI mode if not in test mode
 	if (bRunGuiMode)
diff --git a/Engine/Source/Programs/IoStoreDependencyViewer/Private/IoStoreDependencyViewer.h b/Engine/Source/Programs/IoStoreDependencyViewer/Private/IoStoreDependencyViewer.h
--- a/Engine/Source/Programs/IoStoreDependencyViewer/Private/IoStoreDependencyViewer.h
+++ b/Engine/Source/Programs/IoStoreDependencyViewer/Private/IoStoreDependencyViewer.h
@@ -21,3 +21,12 @@ void TestTocToCSVFromCommandLine();
 
 // Baseline test: Downloads full .utoc + .ucas files and exports asset data to CSV
 bool RunTestBaselineUcas();
+
+// Cloud partial download test: Downloads only metadata, uses on-demand fetching, exports to CSV
+bool RunTestCloudPartialDownloadUCas();
+
+// Zen Reader Integration test: Downloads metadata, uses FIoStoreReaderZenBuild, exports to CSV
+bool RunTestZenReaderIntegration();
+
+// Test function for partial download functionality
+void TestPartialDownload();
diff --git a/Engine/Source/Programs/IoStoreDependencyViewer/Private/SCloudDownloadDialog.cpp b/Engine/Source/Programs/IoStoreDependencyViewer/Private/SCloudDownloadDialog.cpp
--- a/Engine/Source/Programs/IoStoreDependencyViewer/Private/SCloudDownloadDialog.cpp
+++ b/Engine/Source/Programs/IoStoreDependencyViewer/Private/SCloudDownloadDialog.cpp
@@ -8,13 +8,18 @@
 #include "Widgets/Input/SButton.h"
 #include "Widgets/Input/SComboBox.h"
 #include "Widgets/Input/SEditableTextBox.h"
+#include "Widgets/Input/SCheckBox.h"
 #include "Framework/Application/SlateApplication.h"
 #include "HAL/PlatformProcess.h"
+#include "HAL/PlatformFileManager.h"
+#include "Misc/CommandLine.h"
 #include "Misc/Paths.h"
 #include "Misc/FileHelper.h"
 #include "Misc/ConfigCacheIni.h"
+#include "Misc/Parse.h"
 #include "Serialization/JsonSerializer.h"
 #include "DesktopPlatformModule.h"
+#include "PartialDownloadPlatformFile.h"
 
 void SCloudDownloadDialog::Construct(const FArguments& InArgs, TSharedPtr<SWindow> InParentWindow)
 {
@@ -25,8 +30,8 @@ void SCloudDownloadDialog::Construct(const FArguments& InArgs, TSharedPtr<SWindo
 	BuildTypeList.Add(MakeShared<FString>(TEXT("packaged-build")));
 	SelectedBuildType = BuildTypeList[0];
 
-	// Set default download directory
-	DownloadDirectory = FPaths::Combine(FPaths::ProjectDir(), TEXT("CookedBuild"), TEXT("IoStoreDependencyViewer"));
+	// Set default download directory (convert to absolute path)
+	DownloadDirectory = FPaths::ConvertRelativePathToFull(FPaths::Combine(FPaths::ProjectDir(), TEXT("CookedBuild"), TEXT("IoStoreDependencyViewer")));
 
 	ChildSlot
 	[
@@ -249,10 +254,37 @@ void SCloudDownloadDialog::Construct(const FArguments& InArgs, TSharedPtr<SWindo
 				]
 			]
 
+			// Partial downloads checkbox
+			+ SVerticalBox::Slot()
+			.AutoHeight()
+			.Padding(0, 8, 0, 0)
+			[
+				SNew(SHorizontalBox)
+
+				+ SHorizontalBox::Slot()
+				.FillWidth(0.3f)
+				[
+					SNew(STextBlock)
+					.Text(FText::GetEmpty())
+				]
+
+				+ SHorizontalBox::Slot()
+				.FillWidth(0.7f)
+				[
+					SAssignNew(UsePartialDownloadsCheckBox, SCheckBox)
+					.Content()
+					[
+						SNew(STextBlock)
+						.Text(FText::FromString(TEXT("Use partial downloads (experimental)")))
+					]
+					.ToolTipText(FText::FromString(TEXT("Skip .ucas file download; content will be fetched on-demand from the cloud as needed, reducing bandwidth usage")))
+				]
+			]
+
 			// Status text
 			+ SVerticalBox::Slot()
 			.AutoHeight()
-			.Padding(0, 0, 0, 16)
+			.Padding(0, 16, 0, 16)
 			[
 				SAssignNew(StatusTextBlock, STextBlock)
 				.Text(FText::FromString(TEXT("Ready")))
@@ -363,21 +395,103 @@ static FString EscapeCommandLineArgument(const FString& Argument)
 
 bool SCloudDownloadDialog::LoadCloudSettings()
 {
-	// Get zen.exe path relative to current executable
-	FString ExeDir = FPaths::ConvertRelativePathToFull(FPlatformProcess::BaseDir());
-	ZenExePath = FPaths::Combine(ExeDir, TEXT("zen.exe"));
+	// Search for zen.exe in priority order:
+	// 1. If config setting exists, use that always
+	// 2. Same directory as IoStoreDependencyViewer.exe (packaged/development build)
+	// 3. If -projectdir= is specified, check (ProjectDir)\Engine\Binaries\Win64\zen.exe
+
+	FString CustomZenPath;
+	if (GConfig && GConfig->GetString(TEXT("IoStoreDependencyViewer"), TEXT("ZenExePath"), CustomZenPath, GEngineIni))
+	{
+		// Priority 1: Config setting
+		ZenExePath = FPaths::ConvertRelativePathToFull(CustomZenPath);
+		UE_LOG(LogTemp, Log, TEXT("Using zen.exe path from config: %s"), *ZenExePath);
+	}
+	else
+	{
+		// Priority 2: Same directory as executable
+		FString ExeDir = FPaths::ConvertRelativePathToFull(FPlatformProcess::BaseDir());
+		ZenExePath = FPaths::Combine(ExeDir, TEXT("zen.exe"));
+
+		if (!FPaths::FileExists(ZenExePath))
+		{
+			// Priority 3: Check -projectdir= command line argument
+			FString ProjectDir;
+			if (FParse::Value(FCommandLine::Get(), TEXT("projectdir="), ProjectDir))
+			{
+				ProjectDir = FPaths::ConvertRelativePathToFull(ProjectDir);
+				ZenExePath = FPaths::Combine(ProjectDir, TEXT("Engine"), TEXT("Binaries"), TEXT("Win64"), TEXT("zen.exe"));
+				UE_LOG(LogTemp, Log, TEXT("Checking for zen.exe using -projectdir: %s"), *ZenExePath);
+			}
+		}
+	}
+
+	// Search for OidcToken.exe
+	// 1. Config setting (if provided)
+	// 2. Standard Unreal Engine location relative to executable
+	// 3. If -projectdir= is specified, check project directory
+	// 4. Same directory as zen.exe (for convenience)
+
+	FString CustomOidcPath;
+	if (GConfig && GConfig->GetString(TEXT("IoStoreDependencyViewer"), TEXT("OidcTokenExePath"), CustomOidcPath, GEngineIni))
+	{
+		OidcExePath = FPaths::ConvertRelativePathToFull(CustomOidcPath);
+		UE_LOG(LogTemp, Log, TEXT("Using OidcToken.exe path from config: %s"), *OidcExePath);
+	}
+	else
+	{
+		// Default: Standard Unreal Engine location
+		OidcExePath = FPaths::ConvertRelativePathToFull(
+			FPaths::Combine(FPlatformProcess::BaseDir(), TEXT(".."), TEXT("DotNET"), TEXT("OidcToken"), TEXT("win-x64"), TEXT("OidcToken.exe")));
+
+		if (!FPaths::FileExists(OidcExePath))
+		{
+			// Check -projectdir= for OidcToken.exe
+			FString ProjectDir;
+			bool bCheckedProjectDir = FParse::Value(FCommandLine::Get(), TEXT("projectdir="), ProjectDir);
+			if (bCheckedProjectDir)
+			{
+				ProjectDir = FPaths::ConvertRelativePathToFull(ProjectDir);
+				OidcExePath = FPaths::Combine(ProjectDir, TEXT("Engine"), TEXT("Binaries"), TEXT("DotNET"), TEXT("OidcToken"), TEXT("win-x64"), TEXT("OidcToken.exe"));
+				UE_LOG(LogTemp, Log, TEXT("OidcToken.exe not found in standard location, checking -projectdir: %s"), *OidcExePath);
+			}
 
-	// Get OidcToken.exe path relative to current executable
-	OidcExePath = FPaths::ConvertRelativePathToFull(
-		FPaths::Combine(FPlatformProcess::BaseDir(), TEXT(".."), TEXT("DotNET"), TEXT("OidcToken"), TEXT("win-x64"), TEXT("OidcToken.exe")));
+			if (!FPaths::FileExists(OidcExePath))
+			{
+				// Fallback: Same directory as zen.exe (only if zen.exe exists)
+				if (FPaths::FileExists(ZenExePath))
+				{
+					FString ZenDir = FPaths::GetPath(ZenExePath);
+					OidcExePath = FPaths::Combine(ZenDir, TEXT("OidcToken.exe"));
+					if (bCheckedProjectDir)
+					{
+						UE_LOG(LogTemp, Log, TEXT("OidcToken.exe not found in project directory, checking zen.exe directory: %s"), *OidcExePath);
+					}
+					else
+					{
+						UE_LOG(LogTemp, Log, TEXT("OidcToken.exe not found in standard location, checking zen.exe directory: %s"), *OidcExePath);
+					}
+				}
+			}
+		}
+	}
 
 	// Validate paths
 	if (!FPaths::FileExists(ZenExePath))
 	{
 		UE_LOG(LogTemp, Error, TEXT("zen.exe not found at: %s"), *ZenExePath);
+		UE_LOG(LogTemp, Error, TEXT("Searched locations:"));
+		UE_LOG(LogTemp, Error, TEXT("  1. Config: [IoStoreDependencyViewer] ZenExePath"));
+		UE_LOG(LogTemp, Error, TEXT("  2. Executable directory: %s"), *FPaths::Combine(FPlatformProcess::BaseDir(), TEXT("zen.exe")));
+		FString ProjectDir;
+		if (FParse::Value(FCommandLine::Get(), TEXT("projectdir="), ProjectDir))
+		{
+			UE_LOG(LogTemp, Error, TEXT("  3. Project directory: %s"), *FPaths::Combine(ProjectDir, TEXT("Engine\\Binaries\\Win64\\zen.exe")));
+		}
+
 		if (StatusTextBlock.IsValid())
 		{
-			StatusTextBlock->SetText(FText::FromString(TEXT("Error: zen.exe not found")));
+			StatusTextBlock->SetText(FText::FromString(TEXT("Error: zen.exe not found. Check logs for search locations.")));
 			StatusTextBlock->SetColorAndOpacity(FLinearColor::Red);
 		}
 		return false;
@@ -386,9 +500,11 @@ bool SCloudDownloadDialog::LoadCloudSettings()
 	if (!FPaths::FileExists(OidcExePath))
 	{
 		UE_LOG(LogTemp, Error, TEXT("OidcToken.exe not found at: %s"), *OidcExePath);
+		UE_LOG(LogTemp, Error, TEXT("Note: OidcToken.exe is required for authentication. Set OidcTokenExePath in Engine.ini if using a custom location."));
+
 		if (StatusTextBlock.IsValid())
 		{
-			StatusTextBlock->SetText(FText::FromString(TEXT("Error: OidcToken.exe not found")));
+			StatusTextBlock->SetText(FText::FromString(TEXT("Error: OidcToken.exe not found. Check logs for details.")));
 			StatusTextBlock->SetColorAndOpacity(FLinearColor::Red);
 		}
 		return false;
@@ -759,25 +875,77 @@ bool SCloudDownloadDialog::DownloadBuild(const FCloudBuild& Build)
 		}
 	}
 
-	// Build download command - download .utoc, .ucas, .uondemandtoc, and .umeta files
-	// .ucas files contain the actual chunk data that .utoc files reference
-	// Added --verbose for progress visibility
-	FString Command = FString::Printf(
-		TEXT("builds download --namespace=\"%s\" --bucket=\"%s\" --build-id=\"%s\" --local-path=\"%s\" --host \"%s\" --oidctoken-exe-path \"%s\" --wildcard=\"*.utoc;*.ucas;*.uondemandtoc;*.umeta\" --clean --enable-scavenge --verbose"),
-		*EscapeCommandLineArgument(Build.Namespace),
-		*EscapeCommandLineArgument(Build.BucketId),
-		*EscapeCommandLineArgument(Build.BuildId),
-		*EscapeCommandLineArgument(DownloadDirectory),
-		*EscapeCommandLineArgument(ProxyUrl),
-		*EscapeCommandLineArgument(OidcExePath)
-	);
+	// Check if partial downloads are enabled
+	bool bUsePartialDownloads = UsePartialDownloadsCheckBox.IsValid() && UsePartialDownloadsCheckBox->IsChecked();
+
+	// Determine wildcard based on partial download setting
+	FString Wildcard;
+	FString ExcludeWildcard;
+	if (bUsePartialDownloads)
+	{
+		// For partial downloads: download only metadata files (.utoc, .uondemandtoc, .umeta)
+		// Scope wildcard to IoStore file types to prevent --clean from deleting non-IoStore files
+		Wildcard = TEXT("*.utoc;*.uondemandtoc;*.umeta");
+		ExcludeWildcard = TEXT("");
+		UE_LOG(LogTemp, Log, TEXT("Using partial downloads - downloading metadata only"));
+	}
+	else
+	{
+		// For full downloads: download all IoStore files
+		// Scope wildcard to IoStore file types to prevent --clean from deleting non-IoStore files
+		Wildcard = TEXT("*.utoc;*.ucas;*.uondemandtoc;*.umeta");
+		ExcludeWildcard = TEXT("");
+	}
+
+	// Build download command
+	// --clean: Remove files not in the build (prevents stale files across builds)
+	// --enable-scavenge: Resume interrupted downloads by reusing existing chunks
+	// Note: File locking is prevented by calling ClearAllData() before downloads to close open file handles
+	FString Command;
+	if (!ExcludeWildcard.IsEmpty())
+	{
+		Command = FString::Printf(
+			TEXT("builds download --namespace=\"%s\" --bucket=\"%s\" --build-id=\"%s\" --local-path=\"%s\" --host=\"%s\" --oidctoken-exe-path=\"%s\" --wildcard=\"%s\" --exclude-wildcard=\"%s\" --clean --enable-scavenge --verbose"),
+			*EscapeCommandLineArgument(Build.Namespace),
+			*EscapeCommandLineArgument(Build.BucketId),
+			*EscapeCommandLineArgument(Build.BuildId),
+			*EscapeCommandLineArgument(DownloadDirectory),
+			*EscapeCommandLineArgument(ProxyUrl),
+			*EscapeCommandLineArgument(OidcExePath),
+			*Wildcard,
+			*ExcludeWildcard
+		);
+	}
+	else
+	{
+		Command = FString::Printf(
+			TEXT("builds download --namespace=\"%s\" --bucket=\"%s\" --build-id=\"%s\" --local-path=\"%s\" --host=\"%s\" --oidctoken-exe-path=\"%s\" --wildcard=\"%s\" --clean --enable-scavenge --verbose"),
+			*EscapeCommandLineArgument(Build.Namespace),
+			*EscapeCommandLineArgument(Build.BucketId),
+			*EscapeCommandLineArgument(Build.BuildId),
+			*EscapeCommandLineArgument(DownloadDirectory),
+			*EscapeCommandLineArgument(ProxyUrl),
+			*EscapeCommandLineArgument(OidcExePath),
+			*Wildcard
+		);
+	}
+
+	UE_LOG(LogTemp, Display, TEXT("Executing zen.exe download:"));
+	UE_LOG(LogTemp, Display, TEXT("  Wildcard: %s"), *Wildcard);
+	UE_LOG(LogTemp, Display, TEXT("  Exclude wildcard: %s"), ExcludeWildcard.IsEmpty() ? TEXT("<none>") : *ExcludeWildcard);
+	UE_LOG(LogTemp, Display, TEXT("  Partial downloads enabled: %s"), bUsePartialDownloads ? TEXT("YES") : TEXT("NO"));
+	UE_LOG(LogTemp, Display, TEXT("  Download directory: %s"), *DownloadDirectory);
+	UE_LOG(LogTemp, Display, TEXT("  Full command: %s %s"), *ZenExePath, *Command);
 
 	FString Result;
 	int32 ExitCode = RunZenCommand(Command, Result);
 
+	UE_LOG(LogTemp, Display, TEXT("zen.exe download completed with exit code: %d"), ExitCode);
+
 	if (ExitCode != 0)
 	{
 		UE_LOG(LogTemp, Error, TEXT("Download failed. Exit code: %d"), ExitCode);
+		UE_LOG(LogTemp, Error, TEXT("zen.exe output (last 1000 chars): %s"), *Result.Right(1000));
 		if (StatusTextBlock.IsValid())
 		{
 			StatusTextBlock->SetText(FText::FromString(TEXT("Error: Download failed")));
@@ -786,12 +954,43 @@ bool SCloudDownloadDialog::DownloadBuild(const FCloudBuild& Build)
 		return false;
 	}
 
+	// Verify files were actually downloaded
+	TArray<FString> DownloadedUtocFiles;
+	IFileManager::Get().FindFilesRecursive(DownloadedUtocFiles, *DownloadDirectory, TEXT("*.utoc"), true, false);
+	UE_LOG(LogTemp, Display, TEXT("Found %d .utoc files after download"), DownloadedUtocFiles.Num());
+	if (DownloadedUtocFiles.Num() == 0)
+	{
+		UE_LOG(LogTemp, Error, TEXT("No .utoc files found after download - command may have failed silently"));
+		UE_LOG(LogTemp, Error, TEXT("Download directory: %s"), *DownloadDirectory);
+		UE_LOG(LogTemp, Error, TEXT("Wildcard used: %s"), *Wildcard);
+		UE_LOG(LogTemp, Error, TEXT("Exclude wildcard: %s"), ExcludeWildcard.IsEmpty() ? TEXT("<none>") : *ExcludeWildcard);
+
+		if (StatusTextBlock.IsValid())
+		{
+			StatusTextBlock->SetText(FText::FromString(TEXT("Download failed: No .utoc files found. Check logs for details.")));
+			StatusTextBlock->SetColorAndOpacity(FLinearColor::Red);
+		}
+
+		return false;
+	}
+
+	// Partial downloads are handled by FIoStoreReaderZenBuild in LoadStandardContainers
+	// No need for platform file wrapper - zen reader handles on-demand fetching directly
 	if (StatusTextBlock.IsValid())
 	{
-		StatusTextBlock->SetText(FText::FromString(TEXT("Download completed!")));
+		if (bUsePartialDownloads)
+		{
+			StatusTextBlock->SetText(FText::FromString(TEXT("Metadata download complete! On-demand fetching will be used.")));
+		}
+		else
+		{
+			StatusTextBlock->SetText(FText::FromString(TEXT("Download completed!")));
+		}
 		StatusTextBlock->SetColorAndOpacity(FLinearColor::Green);
 	}
+
 	bDownloadInitiated = true;
+	bPartialDownloadUsed = bUsePartialDownloads;
 
 	return true;
 }
diff --git a/Engine/Source/Programs/IoStoreDependencyViewer/Private/SCloudDownloadDialog.h b/Engine/Source/Programs/IoStoreDependencyViewer/Private/SCloudDownloadDialog.h
--- a/Engine/Source/Programs/IoStoreDependencyViewer/Private/SCloudDownloadDialog.h
+++ b/Engine/Source/Programs/IoStoreDependencyViewer/Private/SCloudDownloadDialog.h
@@ -5,6 +5,7 @@
 #include "CoreMinimal.h"
 #include "Widgets/SCompoundWidget.h"
 #include "Widgets/SWindow.h"
+#include "Widgets/Input/SCheckBox.h"
 #include "Widgets/Input/SComboBox.h"
 #include "Widgets/Input/SEditableTextBox.h"
 
@@ -51,6 +52,27 @@ public:
 	/** Returns true if download was initiated */
 	bool WasDownloadInitiated() const { return bDownloadInitiated; }
 
+	/** Returns true if partial downloads were used */
+	bool WasPartialDownloadUsed() const { return bPartialDownloadUsed; }
+
+	/** Returns zen.exe path */
+	FString GetZenExePath() const { return ZenExePath; }
+
+	/** Returns OidcToken.exe path */
+	FString GetOidcExePath() const { return OidcExePath; }
+
+	/** Returns cloud namespace */
+	FString GetCloudNamespace() const { return SelectedBuild.IsValid() ? SelectedBuild->Namespace : FString(); }
+
+	/** Returns cloud bucket */
+	FString GetCloudBucket() const { return SelectedBuild.IsValid() ? SelectedBuild->BucketId : FString(); }
+
+	/** Returns cloud build ID */
+	FString GetCloudBuildId() const { return SelectedBuild.IsValid() ? SelectedBuild->BuildId : FString(); }
+
+	/** Returns proxy URL */
+	FString GetProxyUrl() const { return ProxyUrl; }
+
 private:
 	// Configuration reading
 	bool LoadCloudSettings();
@@ -123,6 +145,7 @@ private:
 
 	FString DownloadDirectory;
 	bool bDownloadInitiated = false;
+	bool bPartialDownloadUsed = false;
 
 	// UI widgets
 	TSharedPtr<SComboBox<TSharedPtr<FString>>> NamespaceComboBox;
@@ -131,5 +154,6 @@ private:
 	TSharedPtr<SComboBox<TSharedPtr<FString>>> PlatformComboBox;
 	TSharedPtr<SComboBox<TSharedPtr<FCloudBuild>>> BuildComboBox;
 	TSharedPtr<SEditableTextBox> DownloadDirectoryTextBox;
+	TSharedPtr<SCheckBox> UsePartialDownloadsCheckBox;
 	TSharedPtr<STextBlock> StatusTextBlock;
 };
diff --git a/Engine/Source/Programs/IoStoreDependencyViewer/Private/SIoStoreDependencyViewer.cpp b/Engine/Source/Programs/IoStoreDependencyViewer/Private/SIoStoreDependencyViewer.cpp
--- a/Engine/Source/Programs/IoStoreDependencyViewer/Private/SIoStoreDependencyViewer.cpp
+++ b/Engine/Source/Programs/IoStoreDependencyViewer/Private/SIoStoreDependencyViewer.cpp
@@ -4,6 +4,8 @@
 #include "IoStoreDependencyViewer.h"
 #include "SDependencyGraphPanel.h"
 #include "SCloudDownloadDialog.h"
+#include "IoStoreReaderAdapter.h"
+#include "IoStoreReaderZenBuild.h"
 #include "Widgets/Input/SButton.h"
 #include "Widgets/Input/SSearchBox.h"
 #include "Widgets/Input/SComboBox.h"
@@ -16,6 +18,7 @@
 #include "Widgets/Views/STableRow.h"
 #include "DesktopPlatformModule.h"
 #include "IDesktopPlatform.h"
+#include "Misc/MessageDialog.h"
 #include "Framework/Application/SlateApplication.h"
 #include "Framework/MultiBox/MultiBoxBuilder.h"
 #include "Styling/AppStyle.h"
@@ -382,6 +385,9 @@
 		RegisterActiveTimer(0.1f, FWidgetActiveTimerDelegate::CreateLambda(
 			[this, InitialPath](double, float) -> EActiveTimerReturnType
 			{
+				// Clear all data before initial load
+				ClearAllData();
+
 				if (LoadFromDirectory(InitialPath))
 				{
 					UpdateAssetList();
@@ -464,6 +470,9 @@
 
 	if (bOpened)
 	{
+		// Clear all data including stale cloud parameters before loading local directory
+		ClearAllData();
+
 		if (LoadFromDirectory(SelectedDirectory))
 		{
 			UpdateAssetList();
@@ -488,37 +497,17 @@
 
 FReply SIoStoreDependencyViewer::OnClearClicked()
 {
-	AllAssets.Empty();
-	FilteredAssets.Empty();
-	CurrentHardDependencies.Empty();
-	CurrentSoftDependencies.Empty();
-	CurrentReferencers.Empty();
-	ChunkIdToAsset.Empty();
-	PackageIdToAsset.Empty();
-	PackageIdToReferencers.Empty();
-	CurrentSearchText.Empty();
-	IoStoreReaders.Empty();
-	MetaReader.Reset();
-
-	if (AssetListView.IsValid())
-	{
-		AssetListView->RequestListRefresh();
-	}
-	if (SearchBox.IsValid())
-	{
-		SearchBox->SetText(FText::GetEmpty());
-	}
-	if (GraphPanel.IsValid())
-	{
-		GraphPanel->RebuildGraph(nullptr, TMap<FPackageId, TSharedPtr<FIoStoreAssetInfo>>(), 0);
-	}
-
+	ClearAllData();
 	StatusText->SetText(LOCTEXT("Cleared", "Cleared all data"));
 	return FReply::Handled();
 }
 
 FReply SIoStoreDependencyViewer::OnCloudDownloadClicked()
 {
+	// Clear all data BEFORE opening dialog to release file handles
+	// This prevents file locking errors when downloading a second build
+	ClearAllData();
+
 	// Create cloud download dialog
 	TSharedRef<SWindow> CloudDownloadWindow = SNew(SWindow)
 		.Title(LOCTEXT("CloudDownloadTitle", "Cloud Download"))
@@ -537,6 +526,30 @@
 		FString DownloadDir = CloudDownloadDialog->GetDownloadDirectory();
 		if (!DownloadDir.IsEmpty() && FPaths::DirectoryExists(DownloadDir))
 		{
+			// Configure cloud build parameters based on whether partial downloads were used
+			if (CloudDownloadDialog->WasPartialDownloadUsed())
+			{
+				UE_LOG(LogIoStoreDependencyViewer, Display, TEXT("Configuring zen reader for partial downloads"));
+				SetCloudBuildParameters(
+					CloudDownloadDialog->GetZenExePath(),
+					CloudDownloadDialog->GetOidcExePath(),
+					CloudDownloadDialog->GetCloudNamespace(),
+					CloudDownloadDialog->GetCloudBucket(),
+					CloudDownloadDialog->GetCloudBuildId(),
+					CloudDownloadDialog->GetProxyUrl(),
+					DownloadDir,  // Base download directory for calculating relative paths
+					true  // bUsePartialDownloads
+				);
+			}
+			else
+			{
+				// Full download - disable partial downloads to use standard FIoStoreReader
+				UE_LOG(LogIoStoreDependencyViewer, Display, TEXT("Full download - using standard IoStore reader"));
+				bUsePartialDownloads = false;
+				ZenExePath.Empty();
+				OidcExePath.Empty();
+			}
+
 			// Find the directory containing .utoc files (may be in subdirectory)
 			FString ActualDir = FindDirectoryWithTocFiles(DownloadDir);
 			if (!ActualDir.IsEmpty())
@@ -566,6 +579,62 @@
 	return FReply::Handled();
 }
 
+void SIoStoreDependencyViewer::ClearAllData()
+{
+	// Clear all data arrays
+	AllAssets.Empty();
+	FilteredAssets.Empty();
+	CurrentHardDependencies.Empty();
+	CurrentSoftDependencies.Empty();
+	CurrentReferencers.Empty();
+	CurrentSelectedAsset.Reset();
+	ChunkIdToAsset.Empty();
+	PackageIdToAsset.Empty();
+	PackageIdToReferencers.Empty();
+	CurrentSearchText.Empty();
+
+	// Close all file handles to release locks
+	IoStoreReaders.Empty();
+	IoStoreAdapters.Empty();
+	MetaReader.Reset();
+
+	// Reset cloud/partial download parameters to prevent stale state
+	bUsePartialDownloads = false;
+	ZenExePath.Empty();
+	OidcExePath.Empty();
+	CloudNamespace.Empty();
+	CloudBucket.Empty();
+	CloudBuildId.Empty();
+	CloudProxyUrl.Empty();
+	BaseDownloadDirectory.Empty();
+
+	// Clear all UI widgets
+	if (AssetListView.IsValid())
+	{
+		AssetListView->RequestListRefresh();
+	}
+	if (HardDependenciesListView.IsValid())
+	{
+		HardDependenciesListView->RequestListRefresh();
+	}
+	if (SoftDependenciesListView.IsValid())
+	{
+		SoftDependenciesListView->RequestListRefresh();
+	}
+	if (ReferencersListView.IsValid())
+	{
+		ReferencersListView->RequestListRefresh();
+	}
+	if (SearchBox.IsValid())
+	{
+		SearchBox->SetText(FText::GetEmpty());
+	}
+	if (GraphPanel.IsValid())
+	{
+		GraphPanel->RebuildGraph(nullptr, TMap<FPackageId, TSharedPtr<FIoStoreAssetInfo>>(), 0);
+	}
+}
+
 bool SIoStoreDependencyViewer::LoadFromDirectory(const FString& DirectoryPath)
 {
 	if (!FPaths::DirectoryExists(DirectoryPath))
@@ -576,13 +645,8 @@
 
 	UE_LOG(LogIoStoreDependencyViewer, Display, TEXT("Loading IoStore containers from: %s"), *DirectoryPath);
 
-	// Clear existing data
-	AllAssets.Empty();
-	ChunkIdToAsset.Empty();
-	PackageIdToAsset.Empty();
-	PackageIdToReferencers.Empty();
-	IoStoreReaders.Empty();
-	MetaReader.Reset();
+	// Note: ClearAllData() is called by each caller before LoadFromDirectory()
+	// This allows cloud download path to set parameters before loading
 
 	// Load .umeta files first for filename resolution
 	LoadMetaData(DirectoryPath);
@@ -593,6 +657,9 @@
 	// Load on-demand .uondemandtoc containers
 	LoadOnDemandContainers(DirectoryPath);
 
+	// Resolve dependency names now that all containers are loaded
+	ResolveDependencyNames();
+
 	UE_LOG(LogIoStoreDependencyViewer, Display, TEXT("Loaded %d total assets/chunks"), AllAssets.Num());
 
 	return AllAssets.Num() > 0;
@@ -650,6 +717,45 @@
 
 	UE_LOG(LogIoStoreDependencyViewer, Display, TEXT("Found %d .utoc files"), TocFiles.Num());
 
+	// Detect if this is a partial download directory (has .utoc but no .ucas files)
+	// If user is trying to load a partial download via "Load Directory" instead of "Cloud Download",
+	// show a helpful error instead of crashing
+	if (!bUsePartialDownloads || ZenExePath.IsEmpty())
+	{
+		// Check if any .ucas files exist
+		TArray<FString> UcasFiles;
+		PlatformFile.FindFilesRecursively(UcasFiles, *DirectoryPath, TEXT(".ucas"));
+
+		// Exclude PartCache directory from the count
+		int32 RealUcasFiles = 0;
+		for (const FString& UcasFile : UcasFiles)
+		{
+			if (!UcasFile.Contains(TEXT("PartCache")))
+			{
+				RealUcasFiles++;
+			}
+		}
+
+		// If we have .utoc files but no .ucas files (excluding PartCache), this is a partial download
+		if (RealUcasFiles == 0 && TocFiles.Num() > 0)
+		{
+			UE_LOG(LogIoStoreDependencyViewer, Error, TEXT("This directory appears to contain a partial cloud download (no .ucas files found)."));
+			UE_LOG(LogIoStoreDependencyViewer, Error, TEXT("Partial downloads cannot be loaded via 'Load Directory'."));
+			UE_LOG(LogIoStoreDependencyViewer, Error, TEXT("Please use 'Cloud Download' to reload this build, or download the full build (uncheck 'Use partial downloads')."));
+
+			// Show error dialog to user
+			FText ErrorTitle = FText::FromString(TEXT("Cannot Load Partial Download"));
+			FText ErrorMessage = FText::FromString(
+				TEXT("This directory contains a partial cloud download (metadata only, no .ucas files).\n\n")
+				TEXT("Partial downloads cannot be loaded via 'Load Directory'.\n\n")
+				TEXT("Please use 'Cloud Download' to reload this build, or download the full build (uncheck 'Use partial downloads').")
+			);
+			FMessageDialog::Open(EAppMsgType::Ok, ErrorMessage, ErrorTitle);
+
+			return false;
+		}
+	}
+
 	// Load each container
 	for (const FString& TocFile : TocFiles)
 	{
@@ -692,24 +798,64 @@
 		UE_LOG(LogIoStoreDependencyViewer, Display, TEXT("  TOC Version: %d (Latest: %d)"), TocVersion, LatestVersion);
 
 		// Version is compatible, proceed with loading
-		TUniquePtr<FIoStoreReader> Reader = MakeUnique<FIoStoreReader>();
-
 		FString ContainerPath = FPaths::ChangeExtension(TocFile, TEXT(""));
-		FIoStatus Status = Reader->Initialize(*ContainerPath, TMap<FGuid, FAES::FAESKey>());
-
-		if (!Status.IsOk())
-		{
-			UE_LOG(LogIoStoreDependencyViewer, Warning, TEXT("Failed to load container %s: %s"),
-				*TocFile, *Status.ToString());
-			continue;
-		}
-
 		FString ContainerName = FPaths::GetBaseFilename(TocFile);
+
+		TSharedPtr<FIoStoreReaderAdapter> Adapter;
+
+		if (bUsePartialDownloads && !ZenExePath.IsEmpty())
+		{
+			// Use zen build reader for on-demand downloads
+			// Create a unique PartCache directory for this container
+			FString PartCacheDir = FPaths::Combine(DirectoryPath, TEXT("PartCache"), ContainerName);
+
+			UE_LOG(LogIoStoreDependencyViewer, Display, TEXT("  Using partial downloads (zen build reader)"));
+			UE_LOG(LogIoStoreDependencyViewer, Display, TEXT("  PartCache: %s"), *PartCacheDir);
+
+			TUniquePtr<FIoStoreReaderZenBuild> ZenReader = MakeUnique<FIoStoreReaderZenBuild>();
+			FIoStatus Status = ZenReader->Initialize(
+				ContainerPath,
+				TMap<FGuid, FAES::FAESKey>(),
+				ZenExePath,
+				OidcExePath,
+				CloudNamespace,
+				CloudBucket,
+				CloudBuildId,
+				CloudProxyUrl,
+				BaseDownloadDirectory,
+				PartCacheDir
+			);
+
+			if (!Status.IsOk())
+			{
+				UE_LOG(LogIoStoreDependencyViewer, Warning, TEXT("Failed to initialize zen build reader for %s: %s"),
+					*TocFile, *Status.ToString());
+				continue;
+			}
+
+			Adapter = MakeShared<FIoStoreReaderAdapter>(MoveTemp(ZenReader));
+		}
+		else
+		{
+			// Standard file-based reader
+			TUniquePtr<FIoStoreReader> StandardReader = MakeUnique<FIoStoreReader>();
+			FIoStatus Status = StandardReader->Initialize(*ContainerPath, TMap<FGuid, FAES::FAESKey>());
+
+			if (!Status.IsOk())
+			{
+				UE_LOG(LogIoStoreDependencyViewer, Warning, TEXT("Failed to load container %s: %s"),
+					*TocFile, *Status.ToString());
+				continue;
+			}
+
+			Adapter = MakeShared<FIoStoreReaderAdapter>(MoveTemp(StandardReader));
+		}
+
 		UE_LOG(LogIoStoreDependencyViewer, Display, TEXT("Container: %s, Chunks: %d"),
-			*ContainerName, Reader->GetChunkCount());
+			*ContainerName, Adapter->GetChunkCount());
 
 		// Enumerate all chunks in this container
-		Reader->EnumerateChunks([this, &ContainerName](FIoStoreTocChunkInfo&& ChunkInfo)
+		Adapter->EnumerateChunks([this, &ContainerName](FIoStoreTocChunkInfo&& ChunkInfo)
 		{
 			TSharedPtr<FIoStoreAssetInfo> Asset = MakeShared<FIoStoreAssetInfo>();
 			Asset->ChunkId = ChunkInfo.Id;
@@ -719,19 +865,10 @@
 			Asset->bIsCompressed = ChunkInfo.bIsCompressed;
 			Asset->PartitionIndex = ChunkInfo.PartitionIndex;
 
-			// Try to resolve filename from metadata
-			if (MetaReader.IsValid())
-			{
-				TUtf8StringBuilder<256> FilenameBuilder;
-				FUtf8StringView ResolvedContainerName;
-				FUtf8StringView Filename = MetaReader->GetFilename(ChunkInfo.Id, FilenameBuilder, ResolvedContainerName);
-				if (!Filename.IsEmpty())
-				{
-					Asset->FileName = FString(Filename);
-				}
-			}
-
-			// If no filename from metadata, use chunk ID as filename
+			// Filename resolution:
+			// - Package chunks: resolved later in ResolveDependencyNames() via PackageId lookup
+			// - Non-package chunks (bulk data, shaders): use ChunkInfo.FileName or fallback to chunk ID
+			// If no filename available, use chunk ID as filename
 			if (Asset->FileName.IsEmpty())
 			{
 				Asset->FileName = ChunkInfo.FileName;
@@ -748,8 +885,8 @@
 		});
 
 		// Try to load and parse the container header for package dependencies
-		FIoChunkId ContainerHeaderChunkId = CreateIoChunkId(Reader->GetContainerId().Value(), 0, EIoChunkType::ContainerHeader);
-		TIoStatusOr<FIoBuffer> HeaderBuffer = Reader->Read(ContainerHeaderChunkId, FIoReadOptions());
+		FIoChunkId ContainerHeaderChunkId = CreateIoChunkId(Adapter->GetContainerId().Value(), 0, EIoChunkType::ContainerHeader);
+		TIoStatusOr<FIoBuffer> HeaderBuffer = Adapter->Read(ContainerHeaderChunkId, FIoReadOptions());
 		if (HeaderBuffer.IsOk())
 		{
 			FIoContainerHeader Header;
@@ -778,17 +915,23 @@
 			}
 
 			// Only process header if serialization succeeded
-		if (!HeaderReader.IsError() && !HeaderReader.IsCriticalError())
-		{
-			ProcessContainerHeader(Header, ContainerName);
+			if (!HeaderReader.IsError() && !HeaderReader.IsCriticalError())
+			{
+				ProcessContainerHeader(Header, ContainerName);
+				UE_LOG(LogIoStoreDependencyViewer, Display, TEXT("    Container header loaded: %d packages"), Header.PackageIds.Num());
+			}
+			else
+			{
+				UE_LOG(LogIoStoreDependencyViewer, Error, TEXT("    Failed to deserialize container header for %s"), *ContainerName);
+			}
 		}
 		else
 		{
-			UE_LOG(LogIoStoreDependencyViewer, Error, TEXT("  Failed to deserialize container header for %s"), *ContainerName);
-		}
+			UE_LOG(LogIoStoreDependencyViewer, Warning, TEXT("    Failed to read container header for %s: %s"),
+				*ContainerName, *HeaderBuffer.Status().ToString());
 		}
 
-		IoStoreReaders.Add(MoveTemp(Reader));
+		IoStoreAdapters.Add(MoveTemp(Adapter));
 	}
 
 	return true;
@@ -996,12 +1139,12 @@
 				for (const FPackageId& ImportedPackage : StoreEntry->ImportedPackages)
 				{
 					FAssetDependency HardDep(ImportedPackage, false);
-					HardDep.PackageName = LexToString(ImportedPackage);
+					// PackageName will be resolved later by ResolveDependencyNames()
 					(*AssetPtr)->HardDependencies.Add(HardDep);
 
 					// Build reverse dependency map (referencers)
 					FAssetDependency ReverseHardDep(PackageId, false);
-					ReverseHardDep.PackageName = LexToString(PackageId);
+					// PackageName will be resolved later by ResolveDependencyNames()
 					PackageIdToReferencers.FindOrAdd(ImportedPackage).Add(ReverseHardDep);
 				}
 			}
@@ -1052,12 +1195,12 @@
 							{
 								FPackageId SoftRefPackageId = Header.SoftPackageReferences.PackageIds[RefIndex];
 								FAssetDependency SoftDep(SoftRefPackageId, true);
-								SoftDep.PackageName = LexToString(SoftRefPackageId);
+								// PackageName will be resolved later by ResolveDependencyNames()
 								(*AssetPtr)->SoftDependencies.Add(SoftDep);
 
 								// Build reverse dependency map (referencers)
 								FAssetDependency ReverseSoftDep(PackageId, true);
-								ReverseSoftDep.PackageName = LexToString(PackageId);
+								// PackageName will be resolved later by ResolveDependencyNames()
 								PackageIdToReferencers.FindOrAdd(SoftRefPackageId).Add(ReverseSoftDep);
 							}
 						}
@@ -1072,6 +1215,129 @@
 	}
 }
 
+void SIoStoreDependencyViewer::ResolveDependencyNames()
+{
+	UE_LOG(LogIoStoreDependencyViewer, Display, TEXT("Resolving dependency names..."));
+
+	int32 ResolvedCount = 0;
+	int32 TotalDeps = 0;
+
+	// Iterate through all assets and resolve their dependency names
+	for (TSharedPtr<FIoStoreAssetInfo>& Asset : AllAssets)
+	{
+		// Resolve hard dependency names
+		for (FAssetDependency& HardDep : Asset->HardDependencies)
+		{
+			TotalDeps++;
+
+			// Try to find the dependency in PackageIdToAsset
+			TSharedPtr<FIoStoreAssetInfo>* DepAssetPtr = PackageIdToAsset.Find(HardDep.PackageId);
+			if (DepAssetPtr && DepAssetPtr->IsValid())
+			{
+				HardDep.PackageName = (*DepAssetPtr)->FileName;
+				ResolvedCount++;
+			}
+			else if (MetaReader.IsValid())
+			{
+				// Try to get filename from metadata
+				TUtf8StringBuilder<256> FilenameBuilder;
+				FUtf8StringView ResolvedContainerName;
+				FIoChunkId PackageChunkId = CreatePackageDataChunkId(HardDep.PackageId);
+				FUtf8StringView Filename = MetaReader->GetFilename(PackageChunkId, FilenameBuilder, ResolvedContainerName);
+				if (!Filename.IsEmpty())
+				{
+					HardDep.PackageName = FString(Filename);
+					ResolvedCount++;
+				}
+				else
+				{
+					// Fallback: use hex PackageId string for unresolved dependencies
+					HardDep.PackageName = LexToString(HardDep.PackageId);
+				}
+			}
+			else
+			{
+				// Fallback: use hex PackageId string for unresolved dependencies
+				HardDep.PackageName = LexToString(HardDep.PackageId);
+			}
+		}
+
+		// Resolve soft dependency names
+		for (FAssetDependency& SoftDep : Asset->SoftDependencies)
+		{
+			TotalDeps++;
+
+			// Try to find the dependency in PackageIdToAsset
+			TSharedPtr<FIoStoreAssetInfo>* DepAssetPtr = PackageIdToAsset.Find(SoftDep.PackageId);
+			if (DepAssetPtr && DepAssetPtr->IsValid())
+			{
+				SoftDep.PackageName = (*DepAssetPtr)->FileName;
+				ResolvedCount++;
+			}
+			else if (MetaReader.IsValid())
+			{
+				// Try to get filename from metadata
+				TUtf8StringBuilder<256> FilenameBuilder;
+				FUtf8StringView ResolvedContainerName;
+				FIoChunkId PackageChunkId = CreatePackageDataChunkId(SoftDep.PackageId);
+				FUtf8StringView Filename = MetaReader->GetFilename(PackageChunkId, FilenameBuilder, ResolvedContainerName);
+				if (!Filename.IsEmpty())
+				{
+					SoftDep.PackageName = FString(Filename);
+					ResolvedCount++;
+				}
+				else
+				{
+					// Fallback: use hex PackageId string for unresolved dependencies
+					SoftDep.PackageName = LexToString(SoftDep.PackageId);
+				}
+			}
+			else
+			{
+				// Fallback: use hex PackageId string for unresolved dependencies
+				SoftDep.PackageName = LexToString(SoftDep.PackageId);
+			}
+		}
+	}
+
+	// Also resolve names in the reverse dependency map (referencers)
+	for (auto& Pair : PackageIdToReferencers)
+	{
+		for (FAssetDependency& Referencer : Pair.Value)
+		{
+			TSharedPtr<FIoStoreAssetInfo>* RefAssetPtr = PackageIdToAsset.Find(Referencer.PackageId);
+			if (RefAssetPtr && RefAssetPtr->IsValid())
+			{
+				Referencer.PackageName = (*RefAssetPtr)->FileName;
+			}
+			else if (MetaReader.IsValid())
+			{
+				TUtf8StringBuilder<256> FilenameBuilder;
+				FUtf8StringView ResolvedContainerName;
+				FIoChunkId PackageChunkId = CreatePackageDataChunkId(Referencer.PackageId);
+				FUtf8StringView Filename = MetaReader->GetFilename(PackageChunkId, FilenameBuilder, ResolvedContainerName);
+				if (!Filename.IsEmpty())
+				{
+					Referencer.PackageName = FString(Filename);
+				}
+				else
+				{
+					// Fallback: use hex PackageId string for unresolved referencers
+					Referencer.PackageName = LexToString(Referencer.PackageId);
+				}
+			}
+			else
+			{
+				// Fallback: use hex PackageId string for unresolved referencers
+				Referencer.PackageName = LexToString(Referencer.PackageId);
+			}
+		}
+	}
+
+	UE_LOG(LogIoStoreDependencyViewer, Display, TEXT("Resolved %d out of %d dependency names (%.1f%%)"),
+		ResolvedCount, TotalDeps, TotalDeps > 0 ? (100.0f * ResolvedCount / TotalDeps) : 0.0f);
+}
+
 void SIoStoreDependencyViewer::UpdateAssetList()
 {
 	FilteredAssets.Empty();
diff --git a/Engine/Source/Programs/IoStoreDependencyViewer/Private/SIoStoreDependencyViewer.h b/Engine/Source/Programs/IoStoreDependencyViewer/Private/SIoStoreDependencyViewer.h
--- a/Engine/Source/Programs/IoStoreDependencyViewer/Private/SIoStoreDependencyViewer.h
+++ b/Engine/Source/Programs/IoStoreDependencyViewer/Private/SIoStoreDependencyViewer.h
@@ -73,11 +73,13 @@ private:
 	FReply OnCloudDownloadClicked();
 
 	// Loading functions
+	void ClearAllData(); // Clears all data and UI state
 	bool LoadFromDirectory(const FString& DirectoryPath);
 	bool LoadStandardContainers(const FString& DirectoryPath);
 	bool LoadOnDemandContainers(const FString& DirectoryPath);
 	bool LoadMetaData(const FString& DirectoryPath);
 	void ProcessContainerHeader(const FIoContainerHeader& Header, const FString& ContainerName);
+	void ResolveDependencyNames();
 
 	// Asset list callbacks
 	TSharedRef<ITableRow> OnGenerateRowForAssetList(TSharedPtr<FIoStoreAssetInfo> AssetInfo, const TSharedRef<STableViewBase>& OwnerTable);
@@ -161,8 +163,44 @@ private:
 	int32 CurrentReferencerDepth = 0;
 
 	// IoStore data
-	TArray<TUniquePtr<FIoStoreReader>> IoStoreReaders;
+	TArray<TUniquePtr<FIoStoreReader>> IoStoreReaders; // Legacy - keeping for now
+	TArray<TSharedPtr<class FIoStoreReaderAdapter>> IoStoreAdapters; // New adapter-based readers
 	TSharedPtr<FIoContainerMetaReader> MetaReader;
 
+	// Cloud build parameters (for partial downloads)
+	bool bUsePartialDownloads = false;
+	FString ZenExePath;
+	FString OidcExePath;
+	FString CloudNamespace;
+	FString CloudBucket;
+	FString CloudBuildId;
+	FString CloudProxyUrl;
+	FString BaseDownloadDirectory; // Where UTOCs are stored
+
 	static constexpr int32 MaxTreeDepth = 10;
+
+public:
+	/**
+	 * Configure cloud build parameters for partial downloads
+	 * Call this before LoadFromDirectory if using partial downloads
+	 */
+	void SetCloudBuildParameters(
+		const FString& InZenExePath,
+		const FString& InOidcExePath,
+		const FString& InNamespace,
+		const FString& InBucket,
+		const FString& InBuildId,
+		const FString& InProxyUrl,
+		const FString& InBaseDownloadDirectory,
+		bool bInUsePartialDownloads)
+	{
+		ZenExePath = InZenExePath;
+		OidcExePath = InOidcExePath;
+		CloudNamespace = InNamespace;
+		CloudBucket = InBucket;
+		CloudBuildId = InBuildId;
+		CloudProxyUrl = InProxyUrl;
+		BaseDownloadDirectory = InBaseDownloadDirectory;
+		bUsePartialDownloads = bInUsePartialDownloads;
+	}
 };
diff --git a/Engine/Source/Programs/IoStoreDependencyViewer/Private/TestBaselineUcasReading.cpp b/Engine/Source/Programs/IoStoreDependencyViewer/Private/TestBaselineUcasReading.cpp
--- a/Engine/Source/Programs/IoStoreDependencyViewer/Private/TestBaselineUcasReading.cpp
+++ b/Engine/Source/Programs/IoStoreDependencyViewer/Private/TestBaselineUcasReading.cpp
@@ -2,6 +2,7 @@
 
 #include "IoStoreDependencyViewer.h"
 #include "Misc/CommandLine.h"
+#include "Misc/ConfigCacheIni.h"
 #include "Misc/FileHelper.h"
 #include "Misc/Paths.h"
 #include "HAL/PlatformFileManager.h"
@@ -189,9 +190,20 @@ namespace BaselineUcasTest
 
 		// Build zen.exe command to download full build part
 		// zen.exe builds download --host=X --namespace=X --bucket=Y --build-id=Z --local-path=D:\Test --oidctoken-exe-path=...
-		// Use fully qualified path to prevent search-order hijacking (binary planting attack)
-		FString ExeDir = FPaths::ConvertRelativePathToFull(FPlatformProcess::BaseDir());
-		FString ZenExePath = FPaths::Combine(ExeDir, TEXT("zen.exe"));
+		// Check for custom zen.exe path in config first
+		FString ZenExePath;
+		FString CustomZenPath;
+		if (GConfig && GConfig->GetString(TEXT("IoStoreDependencyViewer"), TEXT("ZenExePath"), CustomZenPath, GEngineIni))
+		{
+			ZenExePath = FPaths::ConvertRelativePathToFull(CustomZenPath);
+			UE_LOG(LogBaselineUcasTest, Display, TEXT("Using custom zen.exe path from config: %s"), *ZenExePath);
+		}
+		else
+		{
+			// Default: use fully qualified path relative to executable to prevent search-order hijacking (binary planting attack)
+			FString ExeDir = FPaths::ConvertRelativePathToFull(FPlatformProcess::BaseDir());
+			ZenExePath = FPaths::Combine(ExeDir, TEXT("zen.exe"));
+		}
 
 		// Validate zen.exe exists before attempting to execute
 		if (!FPaths::FileExists(ZenExePath))
diff --git a/Engine/Source/Programs/IoStoreDependencyViewer/Private/TestCloudDownload.cpp b/Engine/Source/Programs/IoStoreDependencyViewer/Private/TestCloudDownload.cpp
--- a/Engine/Source/Programs/IoStoreDependencyViewer/Private/TestCloudDownload.cpp
+++ b/Engine/Source/Programs/IoStoreDependencyViewer/Private/TestCloudDownload.cpp
@@ -240,9 +240,22 @@ void TestCloudDownload(const FString& DownloadPath)
 	UE_LOG(LogIoStoreDependencyViewer, Display, TEXT("=== Testing Cloud Download ==="));
 	UE_LOG(LogIoStoreDependencyViewer, Display, TEXT("Download Path: %s"), *DownloadPath);
 
-	// Get zen.exe and OidcToken.exe paths relative to current executable
-	FString ExeDir = FPaths::ConvertRelativePathToFull(FPlatformProcess::BaseDir());
-	FString ZenExePath = FPaths::Combine(ExeDir, TEXT("zen.exe"));
+	// Get zen.exe path - check for custom path in config first
+	FString ZenExePath;
+	FString CustomZenPath;
+	if (GConfig && GConfig->GetString(TEXT("IoStoreDependencyViewer"), TEXT("ZenExePath"), CustomZenPath, GEngineIni))
+	{
+		ZenExePath = FPaths::ConvertRelativePathToFull(CustomZenPath);
+		UE_LOG(LogIoStoreDependencyViewer, Display, TEXT("Using custom zen.exe path from config: %s"), *ZenExePath);
+	}
+	else
+	{
+		// Default: relative to current executable
+		FString ExeDir = FPaths::ConvertRelativePathToFull(FPlatformProcess::BaseDir());
+		ZenExePath = FPaths::Combine(ExeDir, TEXT("zen.exe"));
+	}
+
+	// Get OidcToken.exe path relative to current executable
 	// Use standard Engine layout: ../DotNET/OidcToken/win-x64/ (same as SCloudDownloadDialog and TestCloudDownloadSpecific)
 	FString OidcExePath = FPaths::ConvertRelativePathToFull(
 		FPaths::Combine(FPlatformProcess::BaseDir(), TEXT(".."), TEXT("DotNET"), TEXT("OidcToken"), TEXT("win-x64"), TEXT("OidcToken.exe")));
@@ -678,9 +691,20 @@ void TestCloudDownloadSpecific(const FString& Namespace, const FString& Bucket,
 	Build.Changelist = TEXT("Unknown");
 
 	// Use the SAME download logic as SCloudDownloadDialog::DownloadBuild
-	// Get zen.exe and OidcToken.exe paths - SAME as SCloudDownloadDialog::LoadCloudSettings
-	FString ExeDir = FPaths::ConvertRelativePathToFull(FPlatformProcess::BaseDir());
-	FString ZenExePath = FPaths::Combine(ExeDir, TEXT("zen.exe"));
+	// Get zen.exe path - check for custom path in config first (SAME as SCloudDownloadDialog::LoadCloudSettings)
+	FString ZenExePath;
+	FString CustomZenPath;
+	if (GConfig && GConfig->GetString(TEXT("IoStoreDependencyViewer"), TEXT("ZenExePath"), CustomZenPath, GEngineIni))
+	{
+		ZenExePath = FPaths::ConvertRelativePathToFull(CustomZenPath);
+		UE_LOG(LogIoStoreDependencyViewer, Display, TEXT("Using custom zen.exe path from config: %s"), *ZenExePath);
+	}
+	else
+	{
+		// Default: relative to current executable
+		FString ExeDir = FPaths::ConvertRelativePathToFull(FPlatformProcess::BaseDir());
+		ZenExePath = FPaths::Combine(ExeDir, TEXT("zen.exe"));
+	}
 	FString OidcExePath = FPaths::ConvertRelativePathToFull(
 		FPaths::Combine(FPlatformProcess::BaseDir(), TEXT(".."), TEXT("DotNET"), TEXT("OidcToken"), TEXT("win-x64"), TEXT("OidcToken.exe")));
 
@@ -831,13 +855,54 @@ void TestTocToCSVFromCommandLine()
 	FString TempDir = FPaths::Combine(FPaths::ProjectIntermediateDir(), TEXT("IoStoreDependencyViewer"), TEXT("TocCsvTest"));
 	IFileManager::Get().MakeDirectory(*TempDir, true);
 
-	// Download TOC files
-	FString ExeDir = FPaths::ConvertRelativePathToFull(FPlatformProcess::BaseDir());
-	FString ZenExePath = FPaths::Combine(ExeDir, TEXT("zen.exe"));
+	// Download TOC files - check for custom zen.exe path in config first
+	FString ZenExePath;
+	FString CustomZenPath;
+	if (GConfig && GConfig->GetString(TEXT("IoStoreDependencyViewer"), TEXT("ZenExePath"), CustomZenPath, GEngineIni))
+	{
+		ZenExePath = FPaths::ConvertRelativePathToFull(CustomZenPath);
+		UE_LOG(LogIoStoreDependencyViewer, Display, TEXT("Using custom zen.exe path from config: %s"), *ZenExePath);
+	}
+	else
+	{
+		// Default: relative to current executable
+		FString ExeDir = FPaths::ConvertRelativePathToFull(FPlatformProcess::BaseDir());
+		ZenExePath = FPaths::Combine(ExeDir, TEXT("zen.exe"));
+	}
 	FString OidcExePath = FPaths::ConvertRelativePathToFull(
 		FPaths::Combine(FPlatformProcess::BaseDir(), TEXT(".."), TEXT("DotNET"), TEXT("OidcToken"), TEXT("win-x64"), TEXT("OidcToken.exe")));
 
-	FString ProxyUrl = TEXT("https://jupiter.devtools.epicgames.com");
+	// Load Engine.ini to get host URL from [StorageServers] Cloud setting
+	FString ProxyUrl;
+	// Read Cloud setting from [StorageServers] section
+	FString CloudSetting;
+	if (GConfig && GConfig->GetString(TEXT("StorageServers"), TEXT("Cloud"), CloudSetting, GEngineIni))
+	{
+		// Parse the Host= field from the Cloud setting
+		int32 HostStartIdx = CloudSetting.Find(TEXT("Host=\""));
+		if (HostStartIdx != INDEX_NONE)
+		{
+			HostStartIdx += 6; // Skip 'Host="'
+			int32 HostEndIdx = CloudSetting.Find(TEXT("\""), ESearchCase::IgnoreCase, ESearchDir::FromStart, HostStartIdx);
+			if (HostEndIdx != INDEX_NONE)
+			{
+				FString HostUrls = CloudSetting.Mid(HostStartIdx, HostEndIdx - HostStartIdx);
+				// Split by semicolon and use first URL
+				TArray<FString> Urls;
+				HostUrls.ParseIntoArray(Urls, TEXT(";"));
+				if (Urls.Num() > 0)
+				{
+					ProxyUrl = Urls[0];
+				}
+			}
+		}
+	}
+
+	// Fallback to default if not found
+	if (ProxyUrl.IsEmpty())
+	{
+		ProxyUrl = TEXT("https://jupiter.devtools.epicgames.com");
+	}
 
 	UE_LOG(LogIoStoreDependencyViewer, Display, TEXT("Downloading TOC files..."));
 
