Added async loading of partial zen data to improve performance.

diff --git a/Engine/Source/Programs/IoStoreDependencyViewer/PackageSettings.ini b/Engine/Source/Programs/IoStoreDependencyViewer/PackageSettings.ini
--- a/Engine/Source/Programs/IoStoreDependencyViewer/PackageSettings.ini
+++ b/Engine/Source/Programs/IoStoreDependencyViewer/PackageSettings.ini
@@ -3,7 +3,7 @@
 OutputPath=D:\Depot\depot\InternalTools\UnrealGameSync\UnrealGameSync\bin\Tools\IoStoreDependencyViewer
 
 # Version number (will be written to version.txt inside the package)
-Version=1.0.3
+Version=1.0.4
 
 # Include experimental zen.exe with partial download support (optional)
 # Set to true to include zen.exe and OidcToken.exe in the package
diff --git a/Engine/Source/Programs/IoStoreDependencyViewer/Private/SCloudDownloadDialog.cpp b/Engine/Source/Programs/IoStoreDependencyViewer/Private/SCloudDownloadDialog.cpp
--- a/Engine/Source/Programs/IoStoreDependencyViewer/Private/SCloudDownloadDialog.cpp
+++ b/Engine/Source/Programs/IoStoreDependencyViewer/Private/SCloudDownloadDialog.cpp
@@ -853,6 +853,9 @@ bool SCloudDownloadDialog::QueryBuilds(const FString& Namespace, const FString&
 
 bool SCloudDownloadDialog::DownloadBuild(const FCloudBuild& Build)
 {
+	// Reset cancel flag at start of download
+	bCancelRequested = false;
+
 	if (StatusTextBlock.IsValid())
 	{
 		StatusTextBlock->SetText(FText::FromString(TEXT("Starting download...")));
@@ -944,12 +947,16 @@ bool SCloudDownloadDialog::DownloadBuild(const FCloudBuild& Build)
 
 	if (ExitCode != 0)
 	{
-		UE_LOG(LogTemp, Error, TEXT("Download failed. Exit code: %d"), ExitCode);
-		UE_LOG(LogTemp, Error, TEXT("zen.exe output (last 1000 chars): %s"), *Result.Right(1000));
-		if (StatusTextBlock.IsValid())
+		// Don't log error or overwrite status if this was a user cancellation
+		if (!bCancelRequested)
 		{
-			StatusTextBlock->SetText(FText::FromString(TEXT("Error: Download failed")));
-			StatusTextBlock->SetColorAndOpacity(FLinearColor::Red);
+			UE_LOG(LogTemp, Error, TEXT("Download failed. Exit code: %d"), ExitCode);
+			UE_LOG(LogTemp, Error, TEXT("zen.exe output (last 1000 chars): %s"), *Result.Right(1000));
+			if (StatusTextBlock.IsValid())
+			{
+				StatusTextBlock->SetText(FText::FromString(TEXT("Error: Download failed")));
+				StatusTextBlock->SetColorAndOpacity(FLinearColor::Red);
+			}
 		}
 		return false;
 	}
@@ -1154,6 +1161,23 @@ int32 SCloudDownloadDialog::RunZenCommand(const FString& Command, FString& OutRe
 			}
 		}
 
+		// Check if user requested cancellation
+		if (bCancelRequested)
+		{
+			UE_LOG(LogTemp, Warning, TEXT("User cancelled download - terminating zen.exe"));
+			FPlatformProcess::TerminateProc(ProcHandle, true);
+			FPlatformProcess::ClosePipe(ReadPipe, 0);  // WritePipe already closed
+			FPlatformProcess::CloseProc(ProcHandle);
+
+			if (StatusTextBlock.IsValid())
+			{
+				StatusTextBlock->SetText(FText::FromString(TEXT("Download cancelled by user")));
+				StatusTextBlock->SetColorAndOpacity(FLinearColor::Red);
+			}
+
+			return -1;
+		}
+
 		// Keep UI responsive
 		// NOTE: Pumping Slate messages can trigger callbacks that might destroy this dialog
 		FSlateApplication::Get().PumpMessages();
@@ -1197,6 +1221,14 @@ int32 SCloudDownloadDialog::RunZenCommand(const FString& Command, FString& OutRe
 	FPlatformProcess::CloseProc(ProcHandle);
 	FPlatformProcess::ClosePipe(ReadPipe, 0);  // WritePipe already closed
 
+	// Give zen.exe background threads time to finish file operations before returning
+	// This prevents file sharing violations when scanning for .utoc files immediately after download
+	if (ReturnCode == 0)
+	{
+		UE_LOG(LogTemp, Display, TEXT("Waiting for zen.exe background operations to complete..."));
+		FPlatformProcess::Sleep(2.0f);
+	}
+
 	return ReturnCode;
 }
 
@@ -1221,9 +1253,11 @@ FReply SCloudDownloadDialog::OnDownloadClicked()
 {
 	if (SelectedBuild.IsValid())
 	{
-		if (DownloadBuild(*SelectedBuild))
+		bool bDownloadSucceeded = DownloadBuild(*SelectedBuild);
+
+		// Close dialog after successful download OR after user cancellation
+		if (bDownloadSucceeded || bCancelRequested)
 		{
-			// Close dialog after successful download
 			if (TSharedPtr<SWindow> ParentWindowPtr = ParentWindow.Pin())
 			{
 				FSlateApplication::Get().RequestDestroyWindow(ParentWindowPtr.ToSharedRef());
@@ -1243,10 +1277,15 @@ FReply SCloudDownloadDialog::OnDownloadClicked()
 
 FReply SCloudDownloadDialog::OnCancelClicked()
 {
+	// Set cancel flag - RunZenCommand will detect this and terminate the process if a download is active
+	bCancelRequested = true;
+
+	// Close the dialog window
 	if (TSharedPtr<SWindow> ParentWindowPtr = ParentWindow.Pin())
 	{
 		FSlateApplication::Get().RequestDestroyWindow(ParentWindowPtr.ToSharedRef());
 	}
+
 	return FReply::Handled();
 }
 
@@ -1341,8 +1380,11 @@ TSharedRef<SWidget> SCloudDownloadDialog::OnGenerateStreamWidget(TSharedPtr<FStr
 
 void SCloudDownloadDialog::OnStreamSelectionChanged(TSharedPtr<FString> NewSelection, ESelectInfo::Type SelectInfo)
 {
+	// Save current platform selection to preserve it if still valid
+	FString PreviousPlatform = SelectedPlatform.IsValid() ? *SelectedPlatform : FString();
+
 	SelectedStream = NewSelection;
-	UpdateAvailablePlatforms();
+	UpdateAvailablePlatforms(PreviousPlatform);
 }
 
 FText SCloudDownloadDialog::GetCurrentStreamText() const
@@ -1367,6 +1409,10 @@ FText SCloudDownloadDialog::GetCurrentPlatformText() const
 
 void SCloudDownloadDialog::UpdateAvailableStreams()
 {
+	// Save current selections to preserve them if still valid
+	FString PreviousStream = SelectedStream.IsValid() ? *SelectedStream : FString();
+	FString PreviousPlatform = SelectedPlatform.IsValid() ? *SelectedPlatform : FString();
+
 	StreamList.Empty();
 	PlatformList.Empty();
 	SelectedStream.Reset();
@@ -1419,7 +1465,21 @@ void SCloudDownloadDialog::UpdateAvailableStreams()
 		StreamList.Add(MakeShared<FString>(Stream));
 	}
 
-	if (StreamList.Num() > 0)
+	// Try to preserve previous selection if it's still valid
+	if (!PreviousStream.IsEmpty())
+	{
+		for (const TSharedPtr<FString>& Stream : StreamList)
+		{
+			if (Stream.IsValid() && Stream->Equals(PreviousStream, ESearchCase::IgnoreCase))
+			{
+				SelectedStream = Stream;
+				break;
+			}
+		}
+	}
+
+	// If previous selection wasn't found or was empty, pick the first one
+	if (!SelectedStream.IsValid() && StreamList.Num() > 0)
 	{
 		SelectedStream = StreamList[0];
 	}
@@ -1430,10 +1490,10 @@ void SCloudDownloadDialog::UpdateAvailableStreams()
 	}
 
 	UE_LOG(LogTemp, Log, TEXT("Found %d streams for buildtype %s"), StreamList.Num(), **SelectedBuildType);
-	UpdateAvailablePlatforms();
+	UpdateAvailablePlatforms(PreviousPlatform);
 }
 
-void SCloudDownloadDialog::UpdateAvailablePlatforms()
+void SCloudDownloadDialog::UpdateAvailablePlatforms(const FString& PreviousPlatform)
 {
 	PlatformList.Empty();
 	SelectedPlatform.Reset();
@@ -1479,7 +1539,21 @@ void SCloudDownloadDialog::UpdateAvailablePlatforms()
 		PlatformList.Add(MakeShared<FString>(Platform));
 	}
 
-	if (PlatformList.Num() > 0)
+	// Try to preserve previous selection if it's still valid
+	if (!PreviousPlatform.IsEmpty())
+	{
+		for (const TSharedPtr<FString>& Platform : PlatformList)
+		{
+			if (Platform.IsValid() && Platform->Equals(PreviousPlatform, ESearchCase::IgnoreCase))
+			{
+				SelectedPlatform = Platform;
+				break;
+			}
+		}
+	}
+
+	// If previous selection wasn't found or was empty, pick the first one
+	if (!SelectedPlatform.IsValid() && PlatformList.Num() > 0)
 	{
 		SelectedPlatform = PlatformList[0];
 	}
diff --git a/Engine/Source/Programs/IoStoreDependencyViewer/Private/SCloudDownloadDialog.h b/Engine/Source/Programs/IoStoreDependencyViewer/Private/SCloudDownloadDialog.h
--- a/Engine/Source/Programs/IoStoreDependencyViewer/Private/SCloudDownloadDialog.h
+++ b/Engine/Source/Programs/IoStoreDependencyViewer/Private/SCloudDownloadDialog.h
@@ -8,6 +8,7 @@
 #include "Widgets/Input/SCheckBox.h"
 #include "Widgets/Input/SComboBox.h"
 #include "Widgets/Input/SEditableTextBox.h"
+#include <atomic>
 
 /** Represents a cloud build */
 struct FCloudBuild
@@ -115,7 +116,7 @@ private:
 
 	// Helper functions for cascading dropdowns
 	void UpdateAvailableStreams();
-	void UpdateAvailablePlatforms();
+	void UpdateAvailablePlatforms(const FString& PreviousPlatform = FString());
 	TArray<FString> GetMatchingBuckets() const;
 
 	// Parent window
@@ -146,6 +147,7 @@ private:
 	FString DownloadDirectory;
 	bool bDownloadInitiated = false;
 	bool bPartialDownloadUsed = false;
+	std::atomic<bool> bCancelRequested{false};  // Set when user clicks Cancel during download
 
 	// UI widgets
 	TSharedPtr<SComboBox<TSharedPtr<FString>>> NamespaceComboBox;
diff --git a/Engine/Source/Programs/IoStoreDependencyViewer/Private/SIoStoreDependencyViewer.cpp b/Engine/Source/Programs/IoStoreDependencyViewer/Private/SIoStoreDependencyViewer.cpp
--- a/Engine/Source/Programs/IoStoreDependencyViewer/Private/SIoStoreDependencyViewer.cpp
+++ b/Engine/Source/Programs/IoStoreDependencyViewer/Private/SIoStoreDependencyViewer.cpp
@@ -6,6 +6,7 @@
 #include "SCloudDownloadDialog.h"
 #include "IoStoreReaderAdapter.h"
 #include "IoStoreReaderZenBuild.h"
+#include "Containers/Ticker.h"
 #include "Widgets/Input/SButton.h"
 #include "Widgets/Input/SSearchBox.h"
 #include "Widgets/Input/SComboBox.h"
@@ -33,6 +34,7 @@
 #include "IO/IoDispatcher.h"
 #include "IO/IoContainerMeta.h"
 #include "IO/Serialization/OnDemandContainerToc.h"
+#include "Tasks/Task.h"
 
 #define LOCTEXT_NAMESPACE "IoStoreDependencyViewer"
 
@@ -90,6 +92,26 @@ void SIoStoreDependencyViewer::Construct(const FArguments& InArgs)
 				.OnClicked(this, &SIoStoreDependencyViewer::OnCloudDownloadClicked)
 			]
 			+ SHorizontalBox::Slot()
+			.AutoWidth()
+			.Padding(16.0f, 2.0f, 2.0f, 2.0f)  // Extra left padding for separation
+			[
+				SNew(SButton)
+				.Text(LOCTEXT("NavigateBack", "<"))
+				.ToolTipText(LOCTEXT("NavigateBackTooltip", "Go to previous asset"))
+				.OnClicked(this, &SIoStoreDependencyViewer::OnNavigateBackClicked)
+				.IsEnabled(this, &SIoStoreDependencyViewer::CanNavigateBack)
+			]
+			+ SHorizontalBox::Slot()
+			.AutoWidth()
+			.Padding(2.0f)
+			[
+				SNew(SButton)
+				.Text(LOCTEXT("NavigateForward", ">"))
+				.ToolTipText(LOCTEXT("NavigateForwardTooltip", "Go to next asset"))
+				.OnClicked(this, &SIoStoreDependencyViewer::OnNavigateForwardClicked)
+				.IsEnabled(this, &SIoStoreDependencyViewer::CanNavigateForward)
+			]
+			+ SHorizontalBox::Slot()
 			.FillWidth(1.0f)
 			.HAlign(HAlign_Right)
 			.Padding(2.0f)
@@ -274,10 +296,10 @@ void SIoStoreDependencyViewer::Construct(const FArguments& InArgs)
 			[
 				SNew(SVerticalBox)
 
-				// Graph depth control and toolbar
+				// Graph depth control and toolbar - first row
 				+ SVerticalBox::Slot()
 				.AutoHeight()
-				.Padding(4.0f)
+				.Padding(4.0f, 4.0f, 4.0f, 2.0f)
 				[
 					SNew(SHorizontalBox)
 					+ SHorizontalBox::Slot()
@@ -322,6 +344,7 @@ void SIoStoreDependencyViewer::Construct(const FArguments& InArgs)
 					]
 					+ SHorizontalBox::Slot()
 					.AutoWidth()
+					.VAlign(VAlign_Center)
 					.Padding(10.0f, 0.0f)
 					[
 						SNew(STextBlock)
@@ -345,10 +368,18 @@ void SIoStoreDependencyViewer::Construct(const FArguments& InArgs)
 						.Text(LOCTEXT("ZoomToFit", "Zoom to Fit"))
 						.OnClicked(this, &SIoStoreDependencyViewer::OnZoomToFitClicked)
 					]
+				]
+
+				// Graph search - second row
+				+ SVerticalBox::Slot()
+				.AutoHeight()
+				.Padding(4.0f, 2.0f, 4.0f, 4.0f)
+				[
+					SNew(SHorizontalBox)
 					+ SHorizontalBox::Slot()
 					.AutoWidth()
 					.VAlign(VAlign_Center)
-					.Padding(10.0f, 0.0f)
+					.Padding(4.0f, 0.0f)
 					[
 						SNew(STextBlock)
 						.Text(LOCTEXT("GraphSearchLabel", "Search:"))
@@ -473,6 +504,7 @@ FReply SIoStoreDependencyViewer::OnLoadDirectoryClicked()
 		// Clear all data including stale cloud parameters before loading local directory
 		ClearAllData();
 
+		// LoadFromDirectory is synchronous with modal progress dialog
 		if (LoadFromDirectory(SelectedDirectory))
 		{
 			UpdateAssetList();
@@ -550,13 +582,18 @@ FReply SIoStoreDependencyViewer::OnCloudDownloadClicked()
 				OidcExePath.Empty();
 			}
 
+			// Give zen.exe time to finish background operations before scanning for files
+			// Prevents file sharing violations when zen is moving files to cache
+			UE_LOG(LogIoStoreDependencyViewer, Display, TEXT("Waiting for file system to stabilize..."));
+			FPlatformProcess::Sleep(1.0f);
+
 			// Find the directory containing .utoc files (may be in subdirectory)
 			FString ActualDir = FindDirectoryWithTocFiles(DownloadDir);
 			if (!ActualDir.IsEmpty())
 			{
 				UE_LOG(LogIoStoreDependencyViewer, Display, TEXT("Auto-loading downloaded files from: %s"), *ActualDir);
 
-				// Use the SAME code path as OnLoadDirectoryClicked
+				// LoadFromDirectory is synchronous with modal progress dialog
 				if (LoadFromDirectory(ActualDir))
 				{
 					UpdateAssetList();
@@ -593,6 +630,10 @@ void SIoStoreDependencyViewer::ClearAllData()
 	PackageIdToReferencers.Empty();
 	CurrentSearchText.Empty();
 
+	// Clear navigation history
+	NavigationHistory.Empty();
+	NavigationHistoryIndex = -1;
+
 	// Close all file handles to release locks
 	IoStoreReaders.Empty();
 	IoStoreAdapters.Empty();
@@ -651,7 +692,7 @@ bool SIoStoreDependencyViewer::LoadFromDirectory(const FString& DirectoryPath)
 	// Load .umeta files first for filename resolution
 	LoadMetaData(DirectoryPath);
 
-	// Load standard .utoc containers
+	// Load standard .utoc containers (shows progress dialog internally)
 	LoadStandardContainers(DirectoryPath);
 
 	// Load on-demand .uondemandtoc containers
@@ -756,133 +797,172 @@ bool SIoStoreDependencyViewer::LoadStandardContainers(const FString& DirectoryPa
 		}
 	}
 
-	// Load each container
+	// Launch all containers in parallel for massive speedup
+	TArray<UE::Tasks::FTask> ContainerTasks;
+	FCriticalSection AdaptersLock;  // Protect shared IoStoreAdapters array
+
 	for (const FString& TocFile : TocFiles)
 	{
-		UE_LOG(LogIoStoreDependencyViewer, Display, TEXT("Loading container: %s"), *TocFile);
-
-		// First, validate the TOC version before attempting to load
-		FIoStoreTocResource TocResource;
-		FIoStatus TocStatus = FIoStoreTocResource::Read(*TocFile, EIoStoreTocReadOptions::Default, TocResource);
-
-		if (!TocStatus.IsOk())
+		// Skip this .utoc if a corresponding .uondemandtoc exists (game client behavior)
+		FString OnDemandTocPath = FPaths::ChangeExtension(TocFile, TEXT(".uondemandtoc"));
+		if (PlatformFile.FileExists(*OnDemandTocPath))
 		{
-			UE_LOG(LogIoStoreDependencyViewer, Error, TEXT("Failed to read TOC header from %s: %s"),
-				*TocFile, *TocStatus.ToString());
+			UE_LOG(LogIoStoreDependencyViewer, Display, TEXT("Skipping %s (on-demand TOC exists: %s)"),
+				*FPaths::GetCleanFilename(TocFile), *FPaths::GetCleanFilename(OnDemandTocPath));
 			continue;
 		}
 
-		// Check if the version is compatible
-		const uint8 TocVersion = TocResource.Header.Version;
-		const uint8 LatestVersion = static_cast<uint8>(EIoStoreTocVersion::Latest);
-
-		if (TocVersion > LatestVersion)
+		// Launch container loading as async task
+		UE::Tasks::FTask Task = UE::Tasks::Launch(TEXT("LoadContainer"), [this, TocFile, DirectoryPath, &AdaptersLock]()
 		{
-			UE_LOG(LogIoStoreDependencyViewer, Error,
-				TEXT("Incompatible IoStore version in %s:\n")
-				TEXT("  Container version: %d\n")
-				TEXT("  Latest supported version: %d\n")
-				TEXT("  This file was created with a newer version of the engine.\n")
-				TEXT("  Please use a build from the same or newer branch."),
-				*TocFile, TocVersion, LatestVersion);
-			continue;
-		}
-		else if (TocVersion < static_cast<uint8>(EIoStoreTocVersion::Initial))
-		{
-			UE_LOG(LogIoStoreDependencyViewer, Error,
-				TEXT("Invalid IoStore version in %s: %d"),
-				*TocFile, TocVersion);
-			continue;
-		}
+			IPlatformFile& PlatformFile = FPlatformFileManager::Get().GetPlatformFile();
 
-		UE_LOG(LogIoStoreDependencyViewer, Display, TEXT("  TOC Version: %d (Latest: %d)"), TocVersion, LatestVersion);
+			UE_LOG(LogIoStoreDependencyViewer, Display, TEXT("Loading container: %s"), *TocFile);
 
-		// Version is compatible, proceed with loading
-		FString ContainerPath = FPaths::ChangeExtension(TocFile, TEXT(""));
-		FString ContainerName = FPaths::GetBaseFilename(TocFile);
+			// First, validate the TOC version before attempting to load
+			FIoStoreTocResource TocResource;
+			FIoStatus TocStatus = FIoStoreTocResource::Read(*TocFile, EIoStoreTocReadOptions::Default, TocResource);
 
-		TSharedPtr<FIoStoreReaderAdapter> Adapter;
+			if (!TocStatus.IsOk())
+			{
+				UE_LOG(LogIoStoreDependencyViewer, Error, TEXT("Failed to read TOC header from %s: %s"),
+					*TocFile, *TocStatus.ToString());
+				return;
+			}
 
-		if (bUsePartialDownloads && !ZenExePath.IsEmpty())
-		{
-			// Use zen build reader for on-demand downloads
-			// Create a unique PartCache directory for this container
-			FString PartCacheDir = FPaths::Combine(DirectoryPath, TEXT("PartCache"), ContainerName);
-
-			UE_LOG(LogIoStoreDependencyViewer, Display, TEXT("  Using partial downloads (zen build reader)"));
-			UE_LOG(LogIoStoreDependencyViewer, Display, TEXT("  PartCache: %s"), *PartCacheDir);
-
-			TUniquePtr<FIoStoreReaderZenBuild> ZenReader = MakeUnique<FIoStoreReaderZenBuild>();
-			FIoStatus Status = ZenReader->Initialize(
-				ContainerPath,
-				TMap<FGuid, FAES::FAESKey>(),
-				ZenExePath,
-				OidcExePath,
-				CloudNamespace,
-				CloudBucket,
-				CloudBuildId,
-				CloudProxyUrl,
-				BaseDownloadDirectory,
-				PartCacheDir
-			);
+			// Check if the version is compatible
+			const uint8 TocVersion = TocResource.Header.Version;
+			const uint8 LatestVersion = static_cast<uint8>(EIoStoreTocVersion::Latest);
 
-			if (!Status.IsOk())
+			if (TocVersion > LatestVersion)
 			{
-				UE_LOG(LogIoStoreDependencyViewer, Warning, TEXT("Failed to initialize zen build reader for %s: %s"),
-					*TocFile, *Status.ToString());
-				continue;
+				UE_LOG(LogIoStoreDependencyViewer, Error,
+					TEXT("Incompatible IoStore version in %s:\n")
+					TEXT("  Container version: %d\n")
+					TEXT("  Latest supported version: %d\n")
+					TEXT("  This file was created with a newer version of the engine.\n")
+					TEXT("  Please use a build from the same or newer branch."),
+					*TocFile, TocVersion, LatestVersion);
+				return;
+			}
+			else if (TocVersion < static_cast<uint8>(EIoStoreTocVersion::Initial))
+			{
+				UE_LOG(LogIoStoreDependencyViewer, Error,
+					TEXT("Invalid IoStore version in %s: %d"),
+					*TocFile, TocVersion);
+				return;
 			}
 
-			Adapter = MakeShared<FIoStoreReaderAdapter>(MoveTemp(ZenReader));
-		}
-		else
-		{
-			// Standard file-based reader
-			TUniquePtr<FIoStoreReader> StandardReader = MakeUnique<FIoStoreReader>();
-			FIoStatus Status = StandardReader->Initialize(*ContainerPath, TMap<FGuid, FAES::FAESKey>());
+			UE_LOG(LogIoStoreDependencyViewer, Display, TEXT("  TOC Version: %d (Latest: %d)"), TocVersion, LatestVersion);
+
+			// Version is compatible, proceed with loading
+			FString ContainerPath = FPaths::ChangeExtension(TocFile, TEXT(""));
+			FString ContainerName = FPaths::GetBaseFilename(TocFile);
+
+			TSharedPtr<FIoStoreReaderAdapter> Adapter;
 
-			if (!Status.IsOk())
+			if (bUsePartialDownloads && !ZenExePath.IsEmpty())
 			{
-				UE_LOG(LogIoStoreDependencyViewer, Warning, TEXT("Failed to load container %s: %s"),
-					*TocFile, *Status.ToString());
-				continue;
+				// Use zen build reader for on-demand downloads
+				// Create a unique PartCache directory for this container
+				FString PartCacheDir = FPaths::Combine(DirectoryPath, TEXT("PartCache"), ContainerName);
+
+				UE_LOG(LogIoStoreDependencyViewer, Display, TEXT("  Using partial downloads (zen build reader)"));
+				UE_LOG(LogIoStoreDependencyViewer, Display, TEXT("  PartCache: %s"), *PartCacheDir);
+
+				TUniquePtr<FIoStoreReaderZenBuild> ZenReader = MakeUnique<FIoStoreReaderZenBuild>();
+				FIoStatus Status = ZenReader->Initialize(
+					ContainerPath,
+					TMap<FGuid, FAES::FAESKey>(),
+					ZenExePath,
+					OidcExePath,
+					CloudNamespace,
+					CloudBucket,
+					CloudBuildId,
+					CloudProxyUrl,
+					BaseDownloadDirectory,
+					PartCacheDir
+				);
+
+				if (!Status.IsOk())
+				{
+					// Check if this is an encryption error (expected for DLC/encrypted content without keys)
+					if (Status.ToString().Contains(TEXT("Missing decryption key")))
+					{
+						UE_LOG(LogIoStoreDependencyViewer, Display, TEXT("Skipping encrypted container %s (no decryption key available)"),
+							*FPaths::GetCleanFilename(TocFile));
+					}
+					else
+					{
+						UE_LOG(LogIoStoreDependencyViewer, Warning, TEXT("Failed to initialize zen build reader for %s: %s"),
+							*TocFile, *Status.ToString());
+					}
+					return;
+				}
+
+				Adapter = MakeShared<FIoStoreReaderAdapter>(MoveTemp(ZenReader));
 			}
+			else
+			{
+				// Standard file-based reader
+				TUniquePtr<FIoStoreReader> StandardReader = MakeUnique<FIoStoreReader>();
+				FIoStatus Status = StandardReader->Initialize(*ContainerPath, TMap<FGuid, FAES::FAESKey>());
 
-			Adapter = MakeShared<FIoStoreReaderAdapter>(MoveTemp(StandardReader));
-		}
+				if (!Status.IsOk())
+				{
+					// Check if this is an encryption error (expected for DLC/encrypted content without keys)
+					if (Status.ToString().Contains(TEXT("Missing decryption key")))
+					{
+						UE_LOG(LogIoStoreDependencyViewer, Display, TEXT("Skipping encrypted container %s (no decryption key available)"),
+							*FPaths::GetCleanFilename(TocFile));
+					}
+					else
+					{
+						UE_LOG(LogIoStoreDependencyViewer, Warning, TEXT("Failed to load container %s: %s"),
+							*TocFile, *Status.ToString());
+					}
+					return;
+				}
+
+				Adapter = MakeShared<FIoStoreReaderAdapter>(MoveTemp(StandardReader));
+			}
 
-		UE_LOG(LogIoStoreDependencyViewer, Display, TEXT("Container: %s, Chunks: %d"),
-			*ContainerName, Adapter->GetChunkCount());
+			UE_LOG(LogIoStoreDependencyViewer, Display, TEXT("Container: %s, Chunks: %d"),
+				*ContainerName, Adapter->GetChunkCount());
 
-		// Enumerate all chunks in this container
-		Adapter->EnumerateChunks([this, &ContainerName](FIoStoreTocChunkInfo&& ChunkInfo)
-		{
-			TSharedPtr<FIoStoreAssetInfo> Asset = MakeShared<FIoStoreAssetInfo>();
-			Asset->ChunkId = ChunkInfo.Id;
-			Asset->ContainerName = ContainerName;
-			Asset->Size = ChunkInfo.Size;
-			Asset->CompressedSize = ChunkInfo.CompressedSize;
-			Asset->bIsCompressed = ChunkInfo.bIsCompressed;
-			Asset->PartitionIndex = ChunkInfo.PartitionIndex;
-
-			// Filename resolution:
-			// - Package chunks: resolved later in ResolveDependencyNames() via PackageId lookup
-			// - Non-package chunks (bulk data, shaders): use ChunkInfo.FileName or fallback to chunk ID
-			// If no filename available, use chunk ID as filename
-			if (Asset->FileName.IsEmpty())
+			// Enumerate all chunks in this container
+			Adapter->EnumerateChunks([this, &ContainerName](FIoStoreTocChunkInfo&& ChunkInfo)
 			{
-				Asset->FileName = ChunkInfo.FileName;
+				TSharedPtr<FIoStoreAssetInfo> Asset = MakeShared<FIoStoreAssetInfo>();
+				Asset->ChunkId = ChunkInfo.Id;
+				Asset->ContainerName = ContainerName;
+				Asset->Size = ChunkInfo.Size;
+				Asset->CompressedSize = ChunkInfo.CompressedSize;
+				Asset->bIsCompressed = ChunkInfo.bIsCompressed;
+				Asset->PartitionIndex = ChunkInfo.PartitionIndex;
+
+				// Filename resolution:
+				// - Package chunks: resolved later in ResolveDependencyNames() via PackageId lookup
+				// - Non-package chunks (bulk data, shaders): use ChunkInfo.FileName or fallback to chunk ID
+				// If no filename available, use chunk ID as filename
 				if (Asset->FileName.IsEmpty())
 				{
-					Asset->FileName = FString::Printf(TEXT("Chunk_%s"), *LexToString(ChunkInfo.Id));
+					Asset->FileName = ChunkInfo.FileName;
+					if (Asset->FileName.IsEmpty())
+					{
+						Asset->FileName = FString::Printf(TEXT("Chunk_%s"), *LexToString(ChunkInfo.Id));
+					}
 				}
-			}
 
-			AllAssets.Add(Asset);
-			ChunkIdToAsset.Add(Asset->ChunkId, Asset);
+				// Thread-safe add to shared asset arrays
+				{
+					FScopeLock Lock(&SharedDataLock);
+					AllAssets.Add(Asset);
+					ChunkIdToAsset.Add(Asset->ChunkId, Asset);
+				}
 
-			return true; // Continue enumeration
-		});
+				return true; // Continue enumeration
+			});
 
 		// Try to load and parse the container header for package dependencies
 		FIoChunkId ContainerHeaderChunkId = CreateIoChunkId(Adapter->GetContainerId().Value(), 0, EIoChunkType::ContainerHeader);
@@ -917,6 +997,8 @@ bool SIoStoreDependencyViewer::LoadStandardContainers(const FString& DirectoryPa
 			// Only process header if serialization succeeded
 			if (!HeaderReader.IsError() && !HeaderReader.IsCriticalError())
 			{
+				// Thread-safe processing (modifies PackageIdToAsset and dependency maps)
+				FScopeLock Lock(&SharedDataLock);
 				ProcessContainerHeader(Header, ContainerName);
 				UE_LOG(LogIoStoreDependencyViewer, Display, TEXT("    Container header loaded: %d packages"), Header.PackageIds.Num());
 			}
@@ -931,7 +1013,98 @@ bool SIoStoreDependencyViewer::LoadStandardContainers(const FString& DirectoryPa
 				*ContainerName, *HeaderBuffer.Status().ToString());
 		}
 
-		IoStoreAdapters.Add(MoveTemp(Adapter));
+		// Thread-safe add to shared adapters array
+		{
+			FScopeLock Lock(&AdaptersLock);
+			IoStoreAdapters.Add(MoveTemp(Adapter));
+		}
+		}); // End of lambda
+
+		ContainerTasks.Add(Task);
+	} // End of for loop
+
+	// Show modal progress dialog while containers load in parallel
+	if (ContainerTasks.Num() > 0)
+	{
+		UE_LOG(LogIoStoreDependencyViewer, Display, TEXT("Loading %d containers in parallel..."), ContainerTasks.Num());
+
+		// Create modal progress window
+		TSharedRef<SWindow> ProgressWindow = SNew(SWindow)
+			.Title(FText::FromString(TEXT("Loading Containers")))
+			.ClientSize(FVector2D(400, 100))
+			.SupportsMaximize(false)
+			.SupportsMinimize(false)
+			.IsTopmostWindow(true);
+
+		TSharedPtr<STextBlock> ProgressText;
+		ProgressWindow->SetContent(
+			SNew(SVerticalBox)
+			+ SVerticalBox::Slot()
+			.Padding(20.0f)
+			.VAlign(VAlign_Center)
+			.HAlign(HAlign_Center)
+			[
+				SAssignNew(ProgressText, STextBlock)
+				.Text(FText::FromString(FString::Printf(TEXT("Loading containers: 0/%d..."), ContainerTasks.Num())))
+			]
+		);
+
+		FSlateApplication::Get().AddWindow(ProgressWindow, true);
+
+		// Force window to front and initial paint
+		ProgressWindow->BringToFront();
+		FSlateApplication::Get().Tick();
+
+		// Pump messages while waiting for tasks to complete
+		const int32 TotalTasks = ContainerTasks.Num();
+		while (true)
+		{
+			// Check if all tasks are complete
+			bool bAllComplete = true;
+			int32 CompletedCount = 0;
+			for (const UE::Tasks::FTask& Task : ContainerTasks)
+			{
+				if (Task.IsCompleted())
+				{
+					CompletedCount++;
+				}
+				else
+				{
+					bAllComplete = false;
+				}
+			}
+
+			// Update progress text
+			if (ProgressText.IsValid())
+			{
+				ProgressText->SetText(FText::FromString(FString::Printf(
+					TEXT("Loading containers: %d/%d..."), CompletedCount, TotalTasks)));
+			}
+
+			if (bAllComplete)
+			{
+				break;
+			}
+
+			// Full Slate frame update to keep UI responsive and painted
+			const float DeltaTime = 0.016f;  // ~60 FPS
+
+			// Tick core ticker (handles any timers)
+			FTSTicker::GetCoreTicker().Tick(DeltaTime);
+
+			// Process OS messages
+			FSlateApplication::Get().PumpMessages();
+
+			// Tick and draw Slate (internally calls DrawWindows)
+			FSlateApplication::Get().Tick();
+
+			FPlatformProcess::Sleep(DeltaTime);
+		}
+
+		// Close progress window
+		FSlateApplication::Get().RequestDestroyWindow(ProgressWindow);
+
+		UE_LOG(LogIoStoreDependencyViewer, Display, TEXT("All containers loaded"));
 	}
 
 	return true;
@@ -957,10 +1130,16 @@ bool SIoStoreDependencyViewer::LoadOnDemandContainers(const FString& DirectoryPa
 
 	UE_LOG(LogIoStoreDependencyViewer, Display, TEXT("Found %d .uondemandtoc files"), OnDemandTocFiles.Num());
 
-	// Load each on-demand TOC using the correct API
+	// Launch all on-demand TOCs in parallel
+	TArray<UE::Tasks::FTask> OnDemandTasks;
+
 	for (const FString& TocFile : OnDemandTocFiles)
 	{
-		UE_LOG(LogIoStoreDependencyViewer, Display, TEXT("Loading on-demand TOC: %s"), *TocFile);
+		UE::Tasks::FTask Task = UE::Tasks::Launch(TEXT("LoadOnDemandTOC"), [this, TocFile, DirectoryPath]()
+		{
+			IPlatformFile& PlatformFile = FPlatformFileManager::Get().GetPlatformFile();
+
+			UE_LOG(LogIoStoreDependencyViewer, Display, TEXT("Loading on-demand TOC: %s"), *TocFile);
 
 		// Use FOnDemandTocReader which handles the new serialization format
 		TIoStatusOr<FOnDemandTocReader> MaybeReader = FOnDemandTocReader::Read(TocFile);
@@ -968,7 +1147,7 @@ bool SIoStoreDependencyViewer::LoadOnDemandContainers(const FString& DirectoryPa
 		{
 			UE_LOG(LogIoStoreDependencyViewer, Warning, TEXT("Failed to load on-demand TOC %s: %s"),
 				*TocFile, *MaybeReader.Status().ToString());
-			continue;
+			return;
 		}
 
 		FOnDemandTocReader Reader = MaybeReader.ConsumeValueOrDie();
@@ -1033,16 +1212,72 @@ bool SIoStoreDependencyViewer::LoadOnDemandContainers(const FString& DirectoryPa
 					Asset->FileName = FString::Printf(TEXT("Chunk_%s"), *LexToString(ChunkId));
 				}
 
-				AllAssets.Add(Asset);
-				ChunkIdToAsset.Add(Asset->ChunkId, Asset);
+				// Thread-safe add to shared asset arrays
+				{
+					FScopeLock Lock(&SharedDataLock);
+					AllAssets.Add(Asset);
+					ChunkIdToAsset.Add(Asset->ChunkId, Asset);
+				}
 			}
 
 			// THEN deserialize the container header to get dependency information
+			FIoBuffer ContainerHeaderData;
+
 			if (ContainerView.ContainerHeaderChunk.GetSize() > 0)
+			{
+				// Container header is embedded in the TOC
+				ContainerHeaderData = ContainerView.ContainerHeaderChunk;
+			}
+			else if (bUsePartialDownloads && !ZenExePath.IsEmpty())
+			{
+				// Container header not embedded - try to fetch it as a chunk using zen reader
+				// Check if there's a corresponding .utoc file for this on-demand container
+				FString ContainerTocPath = FPaths::Combine(DirectoryPath, ContainerName + TEXT(".utoc"));
+				if (PlatformFile.FileExists(*ContainerTocPath))
+				{
+					UE_LOG(LogIoStoreDependencyViewer, Display, TEXT("    Fetching container header chunk from zen for %s"), *ContainerName);
+
+					// Create a temporary zen reader to fetch the container header
+					FString ContainerPath = FPaths::ChangeExtension(ContainerTocPath, TEXT(""));
+					FString TempPartCacheDir = FPaths::Combine(DirectoryPath, TEXT("PartCache"), ContainerName);
+
+					TUniquePtr<FIoStoreReaderZenBuild> TempZenReader = MakeUnique<FIoStoreReaderZenBuild>();
+					FIoStatus TempStatus = TempZenReader->Initialize(
+						ContainerPath,
+						TMap<FGuid, FAES::FAESKey>(),
+						ZenExePath,
+						OidcExePath,
+						CloudNamespace,
+						CloudBucket,
+						CloudBuildId,
+						CloudProxyUrl,
+						BaseDownloadDirectory,
+						TempPartCacheDir
+					);
+
+					if (TempStatus.IsOk())
+					{
+						FIoChunkId ContainerHeaderChunkId = CreateIoChunkId(ContainerView.Header.ContainerId().Value(), 0, EIoChunkType::ContainerHeader);
+						TIoStatusOr<FIoBuffer> HeaderBuffer = TempZenReader->Read(ContainerHeaderChunkId, FIoReadOptions());
+
+						if (HeaderBuffer.IsOk())
+						{
+							ContainerHeaderData = HeaderBuffer.ConsumeValueOrDie();
+							UE_LOG(LogIoStoreDependencyViewer, Display, TEXT("    Successfully fetched container header (%llu bytes)"), ContainerHeaderData.GetSize());
+						}
+						else
+						{
+							UE_LOG(LogIoStoreDependencyViewer, Warning, TEXT("    Failed to fetch container header: %s"), *HeaderBuffer.Status().ToString());
+						}
+					}
+				}
+			}
+
+			if (ContainerHeaderData.GetSize() > 0)
 			{
 				FLargeMemoryReader Ar(
-					ContainerView.ContainerHeaderChunk.GetData(),
-					ContainerView.ContainerHeaderChunk.GetSize());
+					ContainerHeaderData.GetData(),
+					ContainerHeaderData.GetSize());
 
 				FIoContainerHeader ContainerHeader;
 				Ar << ContainerHeader;
@@ -1077,8 +1312,11 @@ bool SIoStoreDependencyViewer::LoadOnDemandContainers(const FString& DirectoryPa
 					UE_LOG(LogIoStoreDependencyViewer, Display, TEXT("    Container header loaded: %d packages"),
 						ContainerHeader.PackageIds.Num());
 
-					// Process the container header to build package dependencies
-					ProcessContainerHeader(ContainerHeader, ContainerName);
+					// Process the container header to build package dependencies (thread-safe)
+					{
+						FScopeLock Lock(&SharedDataLock);
+						ProcessContainerHeader(ContainerHeader, ContainerName);
+					}
 				}
 				else
 				{
@@ -1086,6 +1324,17 @@ bool SIoStoreDependencyViewer::LoadOnDemandContainers(const FString& DirectoryPa
 				}
 			}
 		}
+		}); // End of lambda
+
+		OnDemandTasks.Add(Task);
+	} // End of for loop
+
+	// Wait for all on-demand TOCs to finish loading in parallel
+	if (OnDemandTasks.Num() > 0)
+	{
+		UE_LOG(LogIoStoreDependencyViewer, Display, TEXT("Waiting for %d on-demand TOCs to load in parallel..."), OnDemandTasks.Num());
+		UE::Tasks::Wait(OnDemandTasks);
+		UE_LOG(LogIoStoreDependencyViewer, Display, TEXT("All on-demand TOCs loaded"));
 	}
 
 	return true;
@@ -1105,9 +1354,35 @@ void SIoStoreDependencyViewer::ProcessContainerHeader(const FIoContainerHeader&
 		TSharedPtr<FIoStoreAssetInfo>* AssetPtr = ChunkIdToAsset.Find(PackageChunkId);
 		if (AssetPtr && AssetPtr->IsValid())
 		{
-			// Store package ID and name
+			// Store package ID
 			(*AssetPtr)->PackageId = PackageId;
-			(*AssetPtr)->PackageName = LexToString(PackageId);
+
+			// Resolve package filename from metadata
+			if (MetaReader.IsValid())
+			{
+				TUtf8StringBuilder<256> FilenameBuilder;
+				FUtf8StringView ResolvedContainerName;
+				FUtf8StringView Filename = MetaReader->GetFilename(PackageChunkId, FilenameBuilder, ResolvedContainerName);
+				if (!Filename.IsEmpty())
+				{
+					(*AssetPtr)->FileName = FString(Filename);
+					(*AssetPtr)->PackageName = FString(Filename);
+				}
+				else
+				{
+					// Fallback to PackageId hex string if filename not in metadata
+					FString PackageIdString = LexToString(PackageId);
+					(*AssetPtr)->FileName = PackageIdString;
+					(*AssetPtr)->PackageName = PackageIdString;
+				}
+			}
+			else
+			{
+				// No metadata available, use PackageId hex string
+				FString PackageIdString = LexToString(PackageId);
+				(*AssetPtr)->FileName = PackageIdString;
+				(*AssetPtr)->PackageName = PackageIdString;
+			}
 
 			PackageIdToAsset.Add(PackageId, *AssetPtr);
 
@@ -1702,6 +1977,27 @@ void SIoStoreDependencyViewer::OnAssetListSelectionChanged(TSharedPtr<FIoStoreAs
 {
 	CurrentSelectedAsset = AssetInfo;
 
+	// Add to navigation history if this is a user-initiated selection (not during back/forward navigation)
+	if (AssetInfo.IsValid() && SelectInfo != ESelectInfo::Direct && !bIsNavigating)
+	{
+		// Remove any forward history
+		if (NavigationHistoryIndex >= 0 && NavigationHistoryIndex < NavigationHistory.Num() - 1)
+		{
+			NavigationHistory.RemoveAt(NavigationHistoryIndex + 1, NavigationHistory.Num() - NavigationHistoryIndex - 1);
+		}
+
+		// Add to history
+		NavigationHistory.Add(AssetInfo);
+		NavigationHistoryIndex = NavigationHistory.Num() - 1;
+
+		// Limit history size to 50 items
+		if (NavigationHistory.Num() > 50)
+		{
+			NavigationHistory.RemoveAt(0);
+			NavigationHistoryIndex--;
+		}
+	}
+
 	if (AssetInfo.IsValid())
 	{
 		UpdateAssetInfoPanel(AssetInfo);
@@ -1880,4 +2176,75 @@ void SIoStoreDependencyViewer::OnSearchModeChanged(int32 NewMode)
 	UpdateAssetList();
 }
 
+// Navigation history
+
+FReply SIoStoreDependencyViewer::OnNavigateBackClicked()
+{
+	if (CanNavigateBack())
+	{
+		NavigationHistoryIndex--;
+		TSharedPtr<FIoStoreAssetInfo> Asset = NavigationHistory[NavigationHistoryIndex];
+		NavigateToAsset(Asset, false);  // Don't add to history
+	}
+	return FReply::Handled();
+}
+
+FReply SIoStoreDependencyViewer::OnNavigateForwardClicked()
+{
+	if (CanNavigateForward())
+	{
+		NavigationHistoryIndex++;
+		TSharedPtr<FIoStoreAssetInfo> Asset = NavigationHistory[NavigationHistoryIndex];
+		NavigateToAsset(Asset, false);  // Don't add to history
+	}
+	return FReply::Handled();
+}
+
+bool SIoStoreDependencyViewer::CanNavigateBack() const
+{
+	return NavigationHistoryIndex > 0;
+}
+
+bool SIoStoreDependencyViewer::CanNavigateForward() const
+{
+	return NavigationHistoryIndex >= 0 && NavigationHistoryIndex < NavigationHistory.Num() - 1;
+}
+
+void SIoStoreDependencyViewer::NavigateToAsset(TSharedPtr<FIoStoreAssetInfo> Asset, bool bAddToHistory)
+{
+	if (!Asset.IsValid())
+	{
+		return;
+	}
+
+	// Update history if requested (normal user navigation)
+	if (bAddToHistory && !bIsNavigating)
+	{
+		// Remove any forward history
+		if (NavigationHistoryIndex >= 0 && NavigationHistoryIndex < NavigationHistory.Num() - 1)
+		{
+			NavigationHistory.RemoveAt(NavigationHistoryIndex + 1, NavigationHistory.Num() - NavigationHistoryIndex - 1);
+		}
+
+		// Add to history
+		NavigationHistory.Add(Asset);
+		NavigationHistoryIndex = NavigationHistory.Num() - 1;
+
+		// Limit history size to 50 items
+		if (NavigationHistory.Num() > 50)
+		{
+			NavigationHistory.RemoveAt(0);
+			NavigationHistoryIndex--;
+		}
+	}
+
+	// Navigate (will trigger OnAssetListSelectionChanged)
+	bIsNavigating = !bAddToHistory;  // Prevent adding to history during programmatic navigation
+	AssetListView->SetSelection(Asset, ESelectInfo::Direct);
+	AssetListView->RequestScrollIntoView(Asset);
+	bIsNavigating = false;
+}
+
+// Loading progress tracking
+
 #undef LOCTEXT_NAMESPACE
diff --git a/Engine/Source/Programs/IoStoreDependencyViewer/Private/SIoStoreDependencyViewer.h b/Engine/Source/Programs/IoStoreDependencyViewer/Private/SIoStoreDependencyViewer.h
--- a/Engine/Source/Programs/IoStoreDependencyViewer/Private/SIoStoreDependencyViewer.h
+++ b/Engine/Source/Programs/IoStoreDependencyViewer/Private/SIoStoreDependencyViewer.h
@@ -109,6 +109,13 @@ private:
 	void UpdateAssetInfoPanel(TSharedPtr<FIoStoreAssetInfo> AssetInfo);
 	void BuildDependencyGraph(TSharedPtr<FIoStoreAssetInfo> RootAsset);
 
+	// Navigation
+	FReply OnNavigateBackClicked();
+	FReply OnNavigateForwardClicked();
+	bool CanNavigateBack() const;
+	bool CanNavigateForward() const;
+	void NavigateToAsset(TSharedPtr<FIoStoreAssetInfo> Asset, bool bAddToHistory);
+
 	// Context menu handlers
 	TSharedPtr<SWidget> OnReferencersContextMenuOpening();
 	TSharedPtr<SWidget> OnHardDepsContextMenuOpening();
@@ -145,6 +152,14 @@ private:
 	TMap<FPackageId, TArray<FAssetDependency>> PackageIdToReferencers; // Reverse dependency map
 	TSharedPtr<FIoStoreAssetInfo> CurrentSelectedAsset;
 
+	// Navigation history for back/forward
+	TArray<TSharedPtr<FIoStoreAssetInfo>> NavigationHistory;
+	int32 NavigationHistoryIndex = -1;
+	bool bIsNavigating = false;  // Prevent adding to history during navigation
+
+	// Shared data lock for parallel container loading
+	FCriticalSection SharedDataLock;  // Protects AllAssets, ChunkIdToAsset, PackageIdToAsset, PackageIdToReferencers
+
 	// Search state
 	FString CurrentSearchText;
 	int32 SearchMode = 0; // 0=Name, 1=ChunkId, 2=Tags
