Throttle async zen requests, improve error handling, improve UI responsiveness while loading containers.

#Synced-CL 51662096

#Synced-CL 51662096
#Synced-CL 51662096

diff --git a/Engine/Source/Programs/IoStoreDependencyViewer/Private/SIoStoreDependencyViewer.cpp b/Engine/Source/Programs/IoStoreDependencyViewer/Private/SIoStoreDependencyViewer.cpp
--- a/Engine/Source/Programs/IoStoreDependencyViewer/Private/SIoStoreDependencyViewer.cpp
+++ b/Engine/Source/Programs/IoStoreDependencyViewer/Private/SIoStoreDependencyViewer.cpp
@@ -1,5 +1,6 @@
 // Copyright Epic Games, Inc. All Rights Reserved.
 
+#include "CoreMinimal.h"
 #include "SIoStoreDependencyViewer.h"
 #include "IoStoreDependencyViewer.h"
 #include "SDependencyGraphPanel.h"
@@ -14,12 +15,16 @@
 #include "Widgets/Layout/SBorder.h"
 #include "Widgets/Layout/SBox.h"
 #include "Widgets/Layout/SSplitter.h"
+#include "Widgets/Layout/SScrollBox.h"
+#include "Widgets/SWindow.h"
+#include "Widgets/SBoxPanel.h"
 #include "Widgets/Text/STextBlock.h"
 #include "Widgets/Views/SListView.h"
 #include "Widgets/Views/STableRow.h"
 #include "DesktopPlatformModule.h"
 #include "IDesktopPlatform.h"
 #include "Misc/MessageDialog.h"
+#include "Misc/ConfigCacheIni.h"
 #include "Framework/Application/SlateApplication.h"
 #include "Framework/MultiBox/MultiBoxBuilder.h"
 #include "Styling/AppStyle.h"
@@ -421,9 +426,10 @@ void SIoStoreDependencyViewer::Construct(const FArguments& InArgs)
 
 				if (LoadFromDirectory(InitialPath))
 				{
-					UpdateAssetList();
+					// Asset list is updated internally before dialog closes
+					LoadedBuildSource = FString::Printf(TEXT("Local: %s"), *InitialPath);
 					StatusText->SetText(FText::FromString(FString::Printf(TEXT("Loaded: %s (%d assets)"),
-						*FPaths::GetCleanFilename(InitialPath), AllAssets.Num())));
+						*InitialPath, AllAssets.Num())));
 				}
 				return EActiveTimerReturnType::Stop;
 			}));
@@ -505,14 +511,16 @@ FReply SIoStoreDependencyViewer::OnLoadDirectoryClicked()
 		ClearAllData();
 
 		// LoadFromDirectory is synchronous with modal progress dialog
+		// Asset list is updated internally before dialog closes
 		if (LoadFromDirectory(SelectedDirectory))
 		{
-			UpdateAssetList();
+			LoadedBuildSource = FString::Printf(TEXT("Local: %s"), *SelectedDirectory);
 			StatusText->SetText(FText::FromString(FString::Printf(TEXT("Loaded: %s (%d assets)"),
-				*FPaths::GetCleanFilename(SelectedDirectory), AllAssets.Num())));
+				*SelectedDirectory, AllAssets.Num())));
 		}
 		else
 		{
+			LoadedBuildSource.Empty();
 			StatusText->SetText(LOCTEXT("LoadFailed", "Failed to load IoStore containers"));
 		}
 	}
@@ -594,14 +602,27 @@ FReply SIoStoreDependencyViewer::OnCloudDownloadClicked()
 				UE_LOG(LogIoStoreDependencyViewer, Display, TEXT("Auto-loading downloaded files from: %s"), *ActualDir);
 
 				// LoadFromDirectory is synchronous with modal progress dialog
+				// Asset list is updated internally before dialog closes
 				if (LoadFromDirectory(ActualDir))
 				{
-					UpdateAssetList();
-					StatusText->SetText(FText::FromString(FString::Printf(TEXT("Loaded: %s (%d assets)"),
-						*FPaths::GetCleanFilename(ActualDir), AllAssets.Num())));
+					// Show cloud build info in status
+					if (bUsePartialDownloads && !CloudBuildId.IsEmpty())
+					{
+						LoadedBuildSource = FString::Printf(TEXT("Cloud: %s/%s/%s"),
+							*CloudNamespace, *CloudBucket, *CloudBuildId);
+						StatusText->SetText(FText::FromString(FString::Printf(TEXT("Loaded: %s/%s/%s (%d assets)"),
+							*CloudNamespace, *CloudBucket, *CloudBuildId, AllAssets.Num())));
+					}
+					else
+					{
+						LoadedBuildSource = FString::Printf(TEXT("Cloud (full download): %s"), *ActualDir);
+						StatusText->SetText(FText::FromString(FString::Printf(TEXT("Loaded: %s (%d assets)"),
+							*ActualDir, AllAssets.Num())));
+					}
 				}
 				else
 				{
+					LoadedBuildSource.Empty();
 					StatusText->SetText(LOCTEXT("LoadFailed", "Failed to load IoStore containers"));
 				}
 			}
@@ -648,6 +669,7 @@ void SIoStoreDependencyViewer::ClearAllData()
 	CloudBuildId.Empty();
 	CloudProxyUrl.Empty();
 	BaseDownloadDirectory.Empty();
+	LoadedBuildSource.Empty();
 
 	// Clear all UI widgets
 	if (AssetListView.IsValid())
@@ -693,13 +715,8 @@ bool SIoStoreDependencyViewer::LoadFromDirectory(const FString& DirectoryPath)
 	LoadMetaData(DirectoryPath);
 
 	// Load standard .utoc containers (shows progress dialog internally)
-	LoadStandardContainers(DirectoryPath);
-
-	// Load on-demand .uondemandtoc containers
-	LoadOnDemandContainers(DirectoryPath);
-
-	// Resolve dependency names now that all containers are loaded
-	ResolveDependencyNames();
+	// This now also loads on-demand containers and updates the asset list before closing
+	LoadAllContainersWithProgress(DirectoryPath);
 
 	UE_LOG(LogIoStoreDependencyViewer, Display, TEXT("Loaded %d total assets/chunks"), AllAssets.Num());
 
@@ -742,7 +759,226 @@ bool SIoStoreDependencyViewer::LoadMetaData(const FString& DirectoryPath)
 	return false;
 }
 
-bool SIoStoreDependencyViewer::LoadStandardContainers(const FString& DirectoryPath)
+namespace
+{
+	/** Helper struct for container loading progress tracking */
+	struct FContainerLoadError
+	{
+		FString ContainerName;
+		FString ErrorMessage;
+	};
+
+	/** Helper class to manage progress dialog and UI updates during container loading */
+	class FContainerLoadingProgress
+	{
+	public:
+		FContainerLoadingProgress(int32 TotalContainers)
+			: TotalTocs(TotalContainers)
+			, bLoadingFinished(false)
+		{
+			// Create progress window
+			ProgressWindow = SNew(SWindow)
+				.Title(FText::FromString(TEXT("Loading Containers")))
+				.ClientSize(FVector2D(600, 400))
+				.SupportsMaximize(false)
+				.SupportsMinimize(false)
+				.IsTopmostWindow(true);
+
+			ProgressWindow->SetContent(
+				SNew(SVerticalBox)
+				// Progress text
+				+ SVerticalBox::Slot()
+				.AutoHeight()
+				.Padding(20.0f, 20.0f, 20.0f, 10.0f)
+				.VAlign(VAlign_Top)
+				.HAlign(HAlign_Center)
+				[
+					SAssignNew(ProgressText, STextBlock)
+					.Text(FText::FromString(FString::Printf(TEXT("Loading 0/%d containers (0 active)..."), TotalContainers)))
+				]
+				// Error/status list (scrollable)
+				+ SVerticalBox::Slot()
+				.FillHeight(1.0f)
+				.Padding(20.0f, 0.0f, 20.0f, 10.0f)
+				[
+					SNew(SBorder)
+					.BorderImage(FAppStyle::GetBrush("ToolPanel.GroupBorder"))
+					.Padding(5.0f)
+					[
+						SNew(SScrollBox)
+						+ SScrollBox::Slot()
+						[
+							SAssignNew(ErrorText, STextBlock)
+							.AutoWrapText(true)
+							.Text(FText::FromString(TEXT("")))
+						]
+					]
+				]
+				// OK button (initially disabled)
+				+ SVerticalBox::Slot()
+				.AutoHeight()
+				.Padding(20.0f, 0.0f, 20.0f, 20.0f)
+				.HAlign(HAlign_Center)
+				[
+					SAssignNew(OkButton, SButton)
+					.Text(LOCTEXT("OkButton", "OK"))
+					.IsEnabled(false)
+					.OnClicked_Lambda([this]()
+					{
+						if (ProgressWindow.IsValid())
+						{
+							ProgressWindow->RequestDestroyWindow();
+						}
+						return FReply::Handled();
+					})
+				]
+			);
+
+			FSlateApplication::Get().AddWindow(ProgressWindow.ToSharedRef(), true);
+			ProgressWindow->BringToFront();
+			TickUI();
+		}
+
+		void SetProgressText(const FString& Text)
+		{
+			if (ProgressText.IsValid())
+			{
+				ProgressText->SetText(FText::FromString(Text));
+			}
+		}
+
+		void AddError(const FString& ContainerName, const FString& ErrorMessage)
+		{
+			FScopeLock Lock(&ErrorsLock);
+			LoadErrors.Add({ContainerName, ErrorMessage});
+		}
+
+		void AddStatusMessage(const FString& Message)
+		{
+			FScopeLock Lock(&ErrorsLock);
+			StatusLog.Add(Message);
+		}
+
+		void UpdateDisplay(int32 CompletedCount, int32 ActiveCount)
+		{
+			// Update progress text
+			if (!bLoadingFinished)
+			{
+				SetProgressText(FString::Printf(
+					TEXT("Loading %d/%d containers (%d active, %d errors)..."),
+					CompletedCount, TotalTocs, ActiveCount, LoadErrors.Num()));
+			}
+
+			// Update error/status log
+			UpdateErrorStatusDisplay();
+
+			// Enable OK button when finished
+			if (bLoadingFinished && OkButton.IsValid())
+			{
+				OkButton->SetEnabled(true);
+			}
+
+			TickUI();
+		}
+
+		void SetFinished()
+		{
+			bLoadingFinished = true;
+		}
+
+		void TickUI()
+		{
+			const float DeltaTime = 0.016f;  // ~60 FPS
+			FTSTicker::GetCoreTicker().Tick(DeltaTime);
+			FSlateApplication::Get().PumpMessages();
+			FSlateApplication::Get().Tick();
+		}
+
+		void WaitForUserToClose()
+		{
+			while (ProgressWindow.IsValid() && ProgressWindow->IsVisible())
+			{
+				TickUI();
+				FPlatformProcess::Sleep(0.016f);
+			}
+		}
+
+		bool AreAllTasksComplete(const TArray<UE::Tasks::FTask>& Tasks) const
+		{
+			for (const UE::Tasks::FTask& Task : Tasks)
+			{
+				if (!Task.IsCompleted())
+				{
+					return false;
+				}
+			}
+			return true;
+		}
+
+		TSharedPtr<SWindow> GetWindow() const { return ProgressWindow; }
+		int32 GetErrorCount() const { return LoadErrors.Num(); }
+		int32 GetTotalContainers() const { return TotalTocs; }
+
+	private:
+		void UpdateErrorStatusDisplay()
+		{
+			if (!ErrorText.IsValid())
+			{
+				return;
+			}
+
+			FString LogDisplay;
+			{
+				FScopeLock Lock(&ErrorsLock);
+
+				// Show errors first
+				if (LoadErrors.Num() > 0)
+				{
+					LogDisplay += TEXT("=== ERRORS ===\n\n");
+					for (const FContainerLoadError& Error : LoadErrors)
+					{
+						LogDisplay += FString::Printf(TEXT("[%s]: %s\n\n"), *Error.ContainerName, *Error.ErrorMessage);
+					}
+				}
+
+				// Show status messages
+				if (StatusLog.Num() > 0)
+				{
+					if (LoadErrors.Num() > 0)
+					{
+						LogDisplay += TEXT("\n");
+					}
+					LogDisplay += TEXT("=== STATUS ===\n\n");
+					for (const FString& StatusMsg : StatusLog)
+					{
+						LogDisplay += StatusMsg + TEXT("\n");
+					}
+				}
+
+				// If nothing to show yet, show a placeholder
+				if (LoadErrors.Num() == 0 && StatusLog.Num() == 0)
+				{
+					LogDisplay = TEXT("Loading...");
+				}
+			}
+			ErrorText->SetText(FText::FromString(LogDisplay));
+		}
+
+		TSharedPtr<SWindow> ProgressWindow;
+		TSharedPtr<STextBlock> ProgressText;
+		TSharedPtr<STextBlock> ErrorText;
+		TSharedPtr<SButton> OkButton;
+
+		TArray<FContainerLoadError> LoadErrors;
+		TArray<FString> StatusLog;
+		FCriticalSection ErrorsLock;
+
+		int32 TotalTocs;
+		bool bLoadingFinished;
+	};
+}
+
+bool SIoStoreDependencyViewer::LoadAllContainersWithProgress(const FString& DirectoryPath)
 {
 	IPlatformFile& PlatformFile = FPlatformFileManager::Get().GetPlatformFile();
 
@@ -752,70 +988,116 @@ bool SIoStoreDependencyViewer::LoadStandardContainers(const FString& DirectoryPa
 
 	if (TocFiles.Num() == 0)
 	{
-		UE_LOG(LogIoStoreDependencyViewer, Warning, TEXT("No .utoc files found in %s"), *DirectoryPath);
-		return false;
+		UE_LOG(LogIoStoreDependencyViewer, Display, TEXT("No .utoc files found in %s, will try .uondemandtoc files"), *DirectoryPath);
 	}
-
-	UE_LOG(LogIoStoreDependencyViewer, Display, TEXT("Found %d .utoc files"), TocFiles.Num());
-
-	// Detect if this is a partial download directory (has .utoc but no .ucas files)
-	// If user is trying to load a partial download via "Load Directory" instead of "Cloud Download",
-	// show a helpful error instead of crashing
-	if (!bUsePartialDownloads || ZenExePath.IsEmpty())
+	else
 	{
-		// Check if any .ucas files exist
-		TArray<FString> UcasFiles;
-		PlatformFile.FindFilesRecursively(UcasFiles, *DirectoryPath, TEXT(".ucas"));
+		UE_LOG(LogIoStoreDependencyViewer, Display, TEXT("Found %d .utoc files"), TocFiles.Num());
+	}
 
-		// Exclude PartCache directory from the count
-		int32 RealUcasFiles = 0;
-		for (const FString& UcasFile : UcasFiles)
+	// Only perform .utoc validation if we found .utoc files
+	if (TocFiles.Num() > 0)
+	{
+		// Detect if this is a partial download directory (has .utoc but no .ucas files)
+		// If user is trying to load a partial download via "Load Directory" instead of "Cloud Download",
+		// show a helpful error instead of crashing
+		if (!bUsePartialDownloads || ZenExePath.IsEmpty())
 		{
-			if (!UcasFile.Contains(TEXT("PartCache")))
+			// Check if any .ucas files exist
+			TArray<FString> UcasFiles;
+			PlatformFile.FindFilesRecursively(UcasFiles, *DirectoryPath, TEXT(".ucas"));
+
+			// Exclude PartCache directory from the count
+			int32 RealUcasFiles = 0;
+			for (const FString& UcasFile : UcasFiles)
 			{
-				RealUcasFiles++;
+				if (!UcasFile.Contains(TEXT("PartCache")))
+				{
+					RealUcasFiles++;
+				}
 			}
-		}
 
-		// If we have .utoc files but no .ucas files (excluding PartCache), this is a partial download
-		if (RealUcasFiles == 0 && TocFiles.Num() > 0)
-		{
-			UE_LOG(LogIoStoreDependencyViewer, Error, TEXT("This directory appears to contain a partial cloud download (no .ucas files found)."));
-			UE_LOG(LogIoStoreDependencyViewer, Error, TEXT("Partial downloads cannot be loaded via 'Load Directory'."));
-			UE_LOG(LogIoStoreDependencyViewer, Error, TEXT("Please use 'Cloud Download' to reload this build, or download the full build (uncheck 'Use partial downloads')."));
-
-			// Show error dialog to user
-			FText ErrorTitle = FText::FromString(TEXT("Cannot Load Partial Download"));
-			FText ErrorMessage = FText::FromString(
-				TEXT("This directory contains a partial cloud download (metadata only, no .ucas files).\n\n")
-				TEXT("Partial downloads cannot be loaded via 'Load Directory'.\n\n")
-				TEXT("Please use 'Cloud Download' to reload this build, or download the full build (uncheck 'Use partial downloads').")
-			);
-			FMessageDialog::Open(EAppMsgType::Ok, ErrorMessage, ErrorTitle);
+			// If we have .utoc files but no .ucas files (excluding PartCache), this is a partial download
+			if (RealUcasFiles == 0)
+			{
+				UE_LOG(LogIoStoreDependencyViewer, Error, TEXT("This directory appears to contain a partial cloud download (no .ucas files found)."));
+				UE_LOG(LogIoStoreDependencyViewer, Error, TEXT("Partial downloads cannot be loaded via 'Load Directory'."));
+				UE_LOG(LogIoStoreDependencyViewer, Error, TEXT("Please use 'Cloud Download' to reload this build, or download the full build (uncheck 'Use partial downloads')."));
+
+				// Show error dialog to user
+				FText ErrorTitle = FText::FromString(TEXT("Cannot Load Partial Download"));
+				FText ErrorMessage = FText::FromString(
+					TEXT("This directory contains a partial cloud download (metadata only, no .ucas files).\n\n")
+					TEXT("Partial downloads cannot be loaded via 'Load Directory'.\n\n")
+					TEXT("Please use 'Cloud Download' to reload this build, or download the full build (uncheck 'Use partial downloads').")
+				);
+				FMessageDialog::Open(EAppMsgType::Ok, ErrorMessage, ErrorTitle);
 
-			return false;
+				return false;
+			}
 		}
 	}
 
-	// Launch all containers in parallel for massive speedup
-	TArray<UE::Tasks::FTask> ContainerTasks;
-	FCriticalSection AdaptersLock;  // Protect shared IoStoreAdapters array
+	// Read MaxConcurrentContainerLoads from config (default 20)
+	int32 MaxConcurrentContainerLoads = 20;
+	if (GConfig)
+	{
+		GConfig->GetInt(TEXT("IoStoreDependencyViewer"), TEXT("MaxConcurrentContainerLoads"), MaxConcurrentContainerLoads, GEngineIni);
+		// Clamp to reasonable range [1, 100]
+		MaxConcurrentContainerLoads = FMath::Clamp(MaxConcurrentContainerLoads, 1, 100);
+	}
+
+	UE_LOG(LogIoStoreDependencyViewer, Display, TEXT("Using MaxConcurrentContainerLoads = %d"), MaxConcurrentContainerLoads);
 
-	for (const FString& TocFile : TocFiles)
+	// Filter TOC files - skip ones with corresponding .uondemandtoc (game client behavior)
+	TArray<FString> FilteredTocFiles;
+	if (TocFiles.Num() > 0)
 	{
-		// Skip this .utoc if a corresponding .uondemandtoc exists (game client behavior)
-		FString OnDemandTocPath = FPaths::ChangeExtension(TocFile, TEXT(".uondemandtoc"));
-		if (PlatformFile.FileExists(*OnDemandTocPath))
+		for (const FString& TocFile : TocFiles)
 		{
-			UE_LOG(LogIoStoreDependencyViewer, Display, TEXT("Skipping %s (on-demand TOC exists: %s)"),
-				*FPaths::GetCleanFilename(TocFile), *FPaths::GetCleanFilename(OnDemandTocPath));
-			continue;
+			FString OnDemandTocPath = FPaths::ChangeExtension(TocFile, TEXT(".uondemandtoc"));
+			if (PlatformFile.FileExists(*OnDemandTocPath))
+			{
+				UE_LOG(LogIoStoreDependencyViewer, Display, TEXT("Skipping %s (on-demand TOC exists: %s)"),
+					*FPaths::GetCleanFilename(TocFile), *FPaths::GetCleanFilename(OnDemandTocPath));
+				continue;
+			}
+			FilteredTocFiles.Add(TocFile);
 		}
+	}
 
-		// Launch container loading as async task
-		UE::Tasks::FTask Task = UE::Tasks::Launch(TEXT("LoadContainer"), [this, TocFile, DirectoryPath, &AdaptersLock]()
+	const int32 TotalTocs = FilteredTocFiles.Num();
+	if (TotalTocs == 0)
+	{
+		UE_LOG(LogIoStoreDependencyViewer, Display, TEXT("No standard containers to load after filtering"));
+	}
+
+	// Create progress dialog helper (even if TotalTocs is 0, we still need it for on-demand containers)
+	FContainerLoadingProgress Progress(TotalTocs > 0 ? TotalTocs : 1);
+
+	// Only load .utoc files if we have any after filtering
+	if (TotalTocs > 0)
+	{
+
+	// Launch containers in parallel with concurrency limit to prevent system overload
+	TArray<UE::Tasks::FTask> AllContainerTasks;  // All tasks for progress tracking
+	TArray<UE::Tasks::FTask> ActiveTasks;  // Currently running tasks for concurrency control
+	FCriticalSection AdaptersLock;  // Protect shared IoStoreAdapters array
+	int32 NextTocIndex = 0;
+
+	// Lambda to create a task for loading a container
+	auto CreateLoadTask = [this, &DirectoryPath, &AdaptersLock, &Progress](const FString& TocFile) -> UE::Tasks::FTask
+	{
+		return UE::Tasks::Launch(TEXT("LoadContainer"), [this, TocFile, DirectoryPath, &AdaptersLock, &Progress]()
 		{
 			IPlatformFile& PlatformFile = FPlatformFileManager::Get().GetPlatformFile();
+			FString ContainerName = FPaths::GetBaseFilename(TocFile);
+
+			// Helper lambda to report errors
+			auto ReportError = [&Progress, &ContainerName](const FString& ErrorMsg)
+			{
+				Progress.AddError(ContainerName, ErrorMsg);
+			};
 
 			UE_LOG(LogIoStoreDependencyViewer, Display, TEXT("Loading container: %s"), *TocFile);
 
@@ -825,8 +1107,9 @@ bool SIoStoreDependencyViewer::LoadStandardContainers(const FString& DirectoryPa
 
 			if (!TocStatus.IsOk())
 			{
-				UE_LOG(LogIoStoreDependencyViewer, Error, TEXT("Failed to read TOC header from %s: %s"),
-					*TocFile, *TocStatus.ToString());
+				FString ErrorMsg = FString::Printf(TEXT("Failed to read TOC header: %s"), *TocStatus.ToString());
+				UE_LOG(LogIoStoreDependencyViewer, Error, TEXT("%s - %s"), *ContainerName, *ErrorMsg);
+				ReportError(ErrorMsg);
 				return;
 			}
 
@@ -836,20 +1119,19 @@ bool SIoStoreDependencyViewer::LoadStandardContainers(const FString& DirectoryPa
 
 			if (TocVersion > LatestVersion)
 			{
-				UE_LOG(LogIoStoreDependencyViewer, Error,
-					TEXT("Incompatible IoStore version in %s:\n")
-					TEXT("  Container version: %d\n")
-					TEXT("  Latest supported version: %d\n")
-					TEXT("  This file was created with a newer version of the engine.\n")
-					TEXT("  Please use a build from the same or newer branch."),
-					*TocFile, TocVersion, LatestVersion);
+				FString ErrorMsg = FString::Printf(
+					TEXT("Incompatible IoStore version: Container version %d, Latest supported %d. ")
+					TEXT("This file was created with a newer version of the engine."),
+					TocVersion, LatestVersion);
+				UE_LOG(LogIoStoreDependencyViewer, Error, TEXT("%s - %s"), *ContainerName, *ErrorMsg);
+				ReportError(ErrorMsg);
 				return;
 			}
 			else if (TocVersion < static_cast<uint8>(EIoStoreTocVersion::Initial))
 			{
-				UE_LOG(LogIoStoreDependencyViewer, Error,
-					TEXT("Invalid IoStore version in %s: %d"),
-					*TocFile, TocVersion);
+				FString ErrorMsg = FString::Printf(TEXT("Invalid IoStore version: %d"), TocVersion);
+				UE_LOG(LogIoStoreDependencyViewer, Error, TEXT("%s - %s"), *ContainerName, *ErrorMsg);
+				ReportError(ErrorMsg);
 				return;
 			}
 
@@ -857,7 +1139,6 @@ bool SIoStoreDependencyViewer::LoadStandardContainers(const FString& DirectoryPa
 
 			// Version is compatible, proceed with loading
 			FString ContainerPath = FPaths::ChangeExtension(TocFile, TEXT(""));
-			FString ContainerName = FPaths::GetBaseFilename(TocFile);
 
 			TSharedPtr<FIoStoreReaderAdapter> Adapter;
 
@@ -889,13 +1170,14 @@ bool SIoStoreDependencyViewer::LoadStandardContainers(const FString& DirectoryPa
 					// Check if this is an encryption error (expected for DLC/encrypted content without keys)
 					if (Status.ToString().Contains(TEXT("Missing decryption key")))
 					{
-						UE_LOG(LogIoStoreDependencyViewer, Display, TEXT("Skipping encrypted container %s (no decryption key available)"),
-							*FPaths::GetCleanFilename(TocFile));
+						UE_LOG(LogIoStoreDependencyViewer, Display, TEXT("Skipping encrypted container %s (no decryption key available)"), *ContainerName);
+						// Don't report encryption errors - these are expected for encrypted DLC
 					}
 					else
 					{
-						UE_LOG(LogIoStoreDependencyViewer, Warning, TEXT("Failed to initialize zen build reader for %s: %s"),
-							*TocFile, *Status.ToString());
+						FString ErrorMsg = FString::Printf(TEXT("Failed to initialize zen build reader: %s"), *Status.ToString());
+						UE_LOG(LogIoStoreDependencyViewer, Warning, TEXT("%s - %s"), *ContainerName, *ErrorMsg);
+						ReportError(ErrorMsg);
 					}
 					return;
 				}
@@ -913,13 +1195,14 @@ bool SIoStoreDependencyViewer::LoadStandardContainers(const FString& DirectoryPa
 					// Check if this is an encryption error (expected for DLC/encrypted content without keys)
 					if (Status.ToString().Contains(TEXT("Missing decryption key")))
 					{
-						UE_LOG(LogIoStoreDependencyViewer, Display, TEXT("Skipping encrypted container %s (no decryption key available)"),
-							*FPaths::GetCleanFilename(TocFile));
+						UE_LOG(LogIoStoreDependencyViewer, Display, TEXT("Skipping encrypted container %s (no decryption key available)"), *ContainerName);
+						// Don't report encryption errors - these are expected for encrypted DLC
 					}
 					else
 					{
-						UE_LOG(LogIoStoreDependencyViewer, Warning, TEXT("Failed to load container %s: %s"),
-							*TocFile, *Status.ToString());
+						FString ErrorMsg = FString::Printf(TEXT("Failed to load container: %s"), *Status.ToString());
+						UE_LOG(LogIoStoreDependencyViewer, Warning, TEXT("%s - %s"), *ContainerName, *ErrorMsg);
+						ReportError(ErrorMsg);
 					}
 					return;
 				}
@@ -1018,94 +1301,230 @@ bool SIoStoreDependencyViewer::LoadStandardContainers(const FString& DirectoryPa
 			FScopeLock Lock(&AdaptersLock);
 			IoStoreAdapters.Add(MoveTemp(Adapter));
 		}
-		}); // End of lambda
+		}); // End of task lambda
+	}; // End of CreateLoadTask lambda
 
-		ContainerTasks.Add(Task);
-	} // End of for loop
+	// Launch tasks in batches to limit concurrent zen requests
+	UE_LOG(LogIoStoreDependencyViewer, Display, TEXT("Loading %d containers with max %d concurrent loads..."),
+		TotalTocs, MaxConcurrentContainerLoads);
+
+	while (NextTocIndex < TotalTocs || ActiveTasks.Num() > 0)
+	{
+		// Launch new tasks up to the concurrency limit
+		while (NextTocIndex < TotalTocs && ActiveTasks.Num() < MaxConcurrentContainerLoads)
+		{
+			const FString& TocFile = FilteredTocFiles[NextTocIndex++];
+			UE::Tasks::FTask Task = CreateLoadTask(TocFile);
+			AllContainerTasks.Add(Task);
+			ActiveTasks.Add(Task);
+		}
 
-	// Show modal progress dialog while containers load in parallel
-	if (ContainerTasks.Num() > 0)
+		// Remove completed tasks from active list
+		for (int32 i = ActiveTasks.Num() - 1; i >= 0; --i)
+		{
+			if (ActiveTasks[i].IsCompleted())
+			{
+				ActiveTasks.RemoveAtSwap(i);
+			}
+		}
+
+		// Update UI with current progress and errors
+		int32 CompletedCount = AllContainerTasks.Num() - ActiveTasks.Num();
+		Progress.UpdateDisplay(CompletedCount, ActiveTasks.Num());
+
+		// Small sleep to prevent tight loop burning CPU
+		FPlatformProcess::Sleep(0.016f);  // ~60 FPS
+	}
+
+	// Wait for all tasks to complete their work (with UI updates)
+	UE_LOG(LogIoStoreDependencyViewer, Display, TEXT("Waiting for all %d container tasks to complete..."), AllContainerTasks.Num());
+
+	Progress.SetProgressText(FString::Printf(TEXT("Finalizing %d containers..."), TotalTocs));
+	Progress.TickUI();
+
+	while (!Progress.AreAllTasksComplete(AllContainerTasks))
 	{
-		UE_LOG(LogIoStoreDependencyViewer, Display, TEXT("Loading %d containers in parallel..."), ContainerTasks.Num());
+		Progress.TickUI();
+		FPlatformProcess::Sleep(0.016f);
+	}
 
-		// Create modal progress window
-		TSharedRef<SWindow> ProgressWindow = SNew(SWindow)
-			.Title(FText::FromString(TEXT("Loading Containers")))
-			.ClientSize(FVector2D(400, 100))
-			.SupportsMaximize(false)
-			.SupportsMinimize(false)
-			.IsTopmostWindow(true);
+	UE_LOG(LogIoStoreDependencyViewer, Display, TEXT("All container tasks completed. Total assets loaded: %d"), AllAssets.Num());
 
-		TSharedPtr<STextBlock> ProgressText;
-		ProgressWindow->SetContent(
-			SNew(SVerticalBox)
-			+ SVerticalBox::Slot()
-			.Padding(20.0f)
-			.VAlign(VAlign_Center)
-			.HAlign(HAlign_Center)
-			[
-				SAssignNew(ProgressText, STextBlock)
-				.Text(FText::FromString(FString::Printf(TEXT("Loading containers: 0/%d..."), ContainerTasks.Num())))
-			]
-		);
+	// All tasks launched and completed - add completion status
+	if (Progress.GetErrorCount() == 0)
+	{
+		Progress.AddStatusMessage(TEXT("All containers loaded successfully!"));
+	}
+	else
+	{
+		Progress.AddStatusMessage(FString::Printf(TEXT("Loaded %d containers with %d errors."),
+			TotalTocs, Progress.GetErrorCount()));
+	}
+
+	// Update UI to show container loading completion
+	Progress.SetProgressText(FString::Printf(
+		TEXT("Loaded %d containers (%d errors). Resolving dependency names..."),
+		TotalTocs, Progress.GetErrorCount()));
+	Progress.TickUI();
+
+	UE_LOG(LogIoStoreDependencyViewer, Display, TEXT("All containers loaded, resolving dependency names..."));
+
+	// Resolve dependency names (this tracks the count internally)
+	int32 TotalDeps = 0;
+	int32 ResolvedCount = 0;
+
+	// Count total dependencies
+	for (const TSharedPtr<FIoStoreAssetInfo>& Asset : AllAssets)
+	{
+		TotalDeps += Asset->HardDependencies.Num() + Asset->SoftDependencies.Num();
+	}
 
-		FSlateApplication::Get().AddWindow(ProgressWindow, true);
+	// Iterate through all assets and resolve their dependency names
+	int32 ProcessedAssets = 0;
+	for (const TSharedPtr<FIoStoreAssetInfo>& Asset : AllAssets)
+	{
+		ProcessedAssets++;
 
-		// Force window to front and initial paint
-		ProgressWindow->BringToFront();
-		FSlateApplication::Get().Tick();
+		// Periodically update UI during name resolution (every 100 assets)
+		if (ProcessedAssets % 100 == 0)
+		{
+			float ProgressPercent = TotalDeps > 0 ? (100.0f * ResolvedCount / TotalDeps) : 0.0f;
+			Progress.SetProgressText(FString::Printf(
+				TEXT("Loaded %d containers (%d errors)\nResolving names: %d/%d assets (%.1f%% resolved)"),
+				Progress.GetTotalContainers(), Progress.GetErrorCount(), ProcessedAssets, AllAssets.Num(), ProgressPercent));
+			Progress.TickUI();
+		}
 
-		// Pump messages while waiting for tasks to complete
-		const int32 TotalTasks = ContainerTasks.Num();
-		while (true)
+		// Resolve hard dependencies
+		for (FAssetDependency& HardDep : Asset->HardDependencies)
 		{
-			// Check if all tasks are complete
-			bool bAllComplete = true;
-			int32 CompletedCount = 0;
-			for (const UE::Tasks::FTask& Task : ContainerTasks)
+			TSharedPtr<FIoStoreAssetInfo>* DepAssetPtr = PackageIdToAsset.Find(HardDep.PackageId);
+			if (DepAssetPtr && DepAssetPtr->IsValid())
 			{
-				if (Task.IsCompleted())
+				HardDep.PackageName = (*DepAssetPtr)->FileName;
+				ResolvedCount++;
+			}
+			else if (MetaReader.IsValid())
+			{
+				TUtf8StringBuilder<256> FilenameBuilder;
+				FUtf8StringView ResolvedContainerName;
+				FIoChunkId PackageChunkId = CreatePackageDataChunkId(HardDep.PackageId);
+				FUtf8StringView Filename = MetaReader->GetFilename(PackageChunkId, FilenameBuilder, ResolvedContainerName);
+				if (!Filename.IsEmpty())
 				{
-					CompletedCount++;
+					HardDep.PackageName = FString(Filename);
+					ResolvedCount++;
 				}
 				else
 				{
-					bAllComplete = false;
+					HardDep.PackageName = LexToString(HardDep.PackageId);
 				}
 			}
+			else
+			{
+				HardDep.PackageName = LexToString(HardDep.PackageId);
+			}
+		}
 
-			// Update progress text
-			if (ProgressText.IsValid())
+		// Resolve soft dependencies
+		for (FAssetDependency& SoftDep : Asset->SoftDependencies)
+		{
+			TSharedPtr<FIoStoreAssetInfo>* DepAssetPtr = PackageIdToAsset.Find(SoftDep.PackageId);
+			if (DepAssetPtr && DepAssetPtr->IsValid())
+			{
+				SoftDep.PackageName = (*DepAssetPtr)->FileName;
+				ResolvedCount++;
+			}
+			else if (MetaReader.IsValid())
+			{
+				TUtf8StringBuilder<256> FilenameBuilder;
+				FUtf8StringView ResolvedContainerName;
+				FIoChunkId PackageChunkId = CreatePackageDataChunkId(SoftDep.PackageId);
+				FUtf8StringView Filename = MetaReader->GetFilename(PackageChunkId, FilenameBuilder, ResolvedContainerName);
+				if (!Filename.IsEmpty())
+				{
+					SoftDep.PackageName = FString(Filename);
+					ResolvedCount++;
+				}
+				else
+				{
+					SoftDep.PackageName = LexToString(SoftDep.PackageId);
+				}
+			}
+			else
 			{
-				ProgressText->SetText(FText::FromString(FString::Printf(
-					TEXT("Loading containers: %d/%d..."), CompletedCount, TotalTasks)));
+				SoftDep.PackageName = LexToString(SoftDep.PackageId);
 			}
+		}
+	}
 
-			if (bAllComplete)
+	// Also resolve names in the reverse dependency map (referencers)
+	for (auto& Pair : PackageIdToReferencers)
+	{
+		for (FAssetDependency& Referencer : Pair.Value)
+		{
+			TSharedPtr<FIoStoreAssetInfo>* RefAssetPtr = PackageIdToAsset.Find(Referencer.PackageId);
+			if (RefAssetPtr && RefAssetPtr->IsValid())
+			{
+				Referencer.PackageName = (*RefAssetPtr)->FileName;
+			}
+			else if (MetaReader.IsValid())
 			{
-				break;
+				TUtf8StringBuilder<256> FilenameBuilder;
+				FUtf8StringView ResolvedContainerName;
+				FIoChunkId PackageChunkId = CreatePackageDataChunkId(Referencer.PackageId);
+				FUtf8StringView Filename = MetaReader->GetFilename(PackageChunkId, FilenameBuilder, ResolvedContainerName);
+				if (!Filename.IsEmpty())
+				{
+					Referencer.PackageName = FString(Filename);
+				}
+				else
+				{
+					Referencer.PackageName = LexToString(Referencer.PackageId);
+				}
 			}
+			else
+			{
+				Referencer.PackageName = LexToString(Referencer.PackageId);
+			}
+		}
+	}
 
-			// Full Slate frame update to keep UI responsive and painted
-			const float DeltaTime = 0.016f;  // ~60 FPS
+	float ResolvePercent = TotalDeps > 0 ? (100.0f * ResolvedCount / TotalDeps) : 0.0f;
+	UE_LOG(LogIoStoreDependencyViewer, Display, TEXT("Resolved %d out of %d dependency names (%.1f%%)"),
+		ResolvedCount, TotalDeps, ResolvePercent);
 
-			// Tick core ticker (handles any timers)
-			FTSTicker::GetCoreTicker().Tick(DeltaTime);
+		// Add name resolution status to log
+		Progress.AddStatusMessage(FString::Printf(TEXT("Resolved %d out of %d dependency names (%.1f%%)"),
+			ResolvedCount, TotalDeps, ResolvePercent));
+	} // End of if (TotalTocs > 0) - standard container loading complete
 
-			// Process OS messages
-			FSlateApplication::Get().PumpMessages();
+	// Load on-demand containers while dialog is still up (always called, even if no .utoc files)
+	Progress.SetProgressText(TEXT("Loading on-demand containers..."));
+	Progress.TickUI();
 
-			// Tick and draw Slate (internally calls DrawWindows)
-			FSlateApplication::Get().Tick();
+	UE_LOG(LogIoStoreDependencyViewer, Display, TEXT("Loading on-demand containers..."));
+	LoadOnDemandContainers(DirectoryPath);
 
-			FPlatformProcess::Sleep(DeltaTime);
-		}
+	// Update the asset list with all loaded assets
+	Progress.SetProgressText(FString::Printf(TEXT("Updating asset list (%d assets)..."), AllAssets.Num()));
+	Progress.TickUI();
+
+	UE_LOG(LogIoStoreDependencyViewer, Display, TEXT("Updating asset list with %d total assets..."), AllAssets.Num());
+	UpdateAssetList();
 
-		// Close progress window
-		FSlateApplication::Get().RequestDestroyWindow(ProgressWindow);
+	// All done - enable OK button and update final status
+	Progress.SetFinished();
+	Progress.SetProgressText(FString::Printf(
+		TEXT("Complete! Loaded %d total assets. Click OK to continue."), AllAssets.Num()));
+	Progress.UpdateDisplay(0, 0);  // Force final UI update to enable OK button
 
-		UE_LOG(LogIoStoreDependencyViewer, Display, TEXT("All containers loaded"));
-	}
+	UE_LOG(LogIoStoreDependencyViewer, Display, TEXT("All loading complete, waiting for user to close dialog"));
+
+	// Keep window open and pump messages until user clicks OK
+	Progress.WaitForUserToClose();
+
+	UE_LOG(LogIoStoreDependencyViewer, Display, TEXT("Container loading dialog closed"));
 
 	return true;
 }
diff --git a/Engine/Source/Programs/IoStoreDependencyViewer/README.md b/Engine/Source/Programs/IoStoreDependencyViewer/README.md
--- a/Engine/Source/Programs/IoStoreDependencyViewer/README.md
+++ b/Engine/Source/Programs/IoStoreDependencyViewer/README.md
@@ -96,6 +96,23 @@ The tool extracts dependencies from:
 | **Filename Resolution** | Via `.umeta` files | Built into registry |
 | **Best For** | Shipping/cooked build analysis | Development/editor workflows |
 
+## Configuration
+
+The tool can be configured via the `Engine.ini` file. Add settings under the `[IoStoreDependencyViewer]` section:
+
+```ini
+[IoStoreDependencyViewer]
+; Maximum number of containers to load concurrently (default: 20)
+; Lower values reduce memory/CPU usage, higher values speed up loading
+; Valid range: 1-100
+MaxConcurrentContainerLoads=20
+```
+
+**MaxConcurrentContainerLoads**: Controls how many container files are loaded in parallel. When loading builds with many containers (100+), this prevents overwhelming the system by limiting concurrent zen.exe requests and file I/O operations. Adjust based on your system specs:
+- **High-end systems** (32+ GB RAM, fast SSD): 30-50
+- **Mid-range systems** (16 GB RAM): 20 (default)
+- **Low-end systems** (8 GB RAM): 10-15
+
 ## Notes
 
 - This tool does NOT modify any engine or game code
