IoStore Dependency Viewer cleans part directories.

#Synced-CL 51662096

diff --git a/Engine/Source/Programs/IoStoreDependencyViewer/Private/SCloudDownloadDialog.cpp b/Engine/Source/Programs/IoStoreDependencyViewer/Private/SCloudDownloadDialog.cpp
--- a/Engine/Source/Programs/IoStoreDependencyViewer/Private/SCloudDownloadDialog.cpp
+++ b/Engine/Source/Programs/IoStoreDependencyViewer/Private/SCloudDownloadDialog.cpp
@@ -981,7 +981,7 @@ bool SCloudDownloadDialog::DownloadBuild(const FCloudBuild& Build)
 		return false;
 	}
 
-	// Partial downloads are handled by FIoStoreReaderZenBuild in LoadStandardContainers
+	// Partial downloads are handled by FIoStoreReaderZenBuild in LoadAllContainersWithProgress
 	// No need for platform file wrapper - zen reader handles on-demand fetching directly
 	if (StatusTextBlock.IsValid())
 	{
@@ -1050,7 +1050,7 @@ int32 SCloudDownloadDialog::RunZenCommand(const FString& Command, FString& OutRe
 	int32 ProgressCounter = 0;
 	int32 LastKnownPercentage = 0;
 
-	// Wait for process to complete with 2-hour timeout
+	// Wait for process to complete with 2-hour timeout (full downloads of large builds can take hours)
 	const double TimeoutSeconds = 2.0 * 60.0 * 60.0;
 	while (FPlatformProcess::IsProcRunning(ProcHandle))
 	{
diff --git a/Engine/Source/Programs/IoStoreDependencyViewer/Private/SIoStoreDependencyViewer.cpp b/Engine/Source/Programs/IoStoreDependencyViewer/Private/SIoStoreDependencyViewer.cpp
--- a/Engine/Source/Programs/IoStoreDependencyViewer/Private/SIoStoreDependencyViewer.cpp
+++ b/Engine/Source/Programs/IoStoreDependencyViewer/Private/SIoStoreDependencyViewer.cpp
@@ -660,6 +660,25 @@ void SIoStoreDependencyViewer::ClearAllData()
 	IoStoreAdapters.Empty();
 	MetaReader.Reset();
 
+	// Clean up PartCache directory before resetting parameters
+	// Zen really struggles with processing stale PartCache files, so delete them on clear
+	if (!BaseDownloadDirectory.IsEmpty())
+	{
+		FString PartCacheDir = FPaths::Combine(BaseDownloadDirectory, TEXT("PartCache"));
+		if (FPaths::DirectoryExists(PartCacheDir))
+		{
+			IPlatformFile& PlatformFile = FPlatformFileManager::Get().GetPlatformFile();
+			if (PlatformFile.DeleteDirectoryRecursively(*PartCacheDir))
+			{
+				UE_LOG(LogIoStoreDependencyViewer, Display, TEXT("Cleaned up PartCache directory: %s"), *PartCacheDir);
+			}
+			else
+			{
+				UE_LOG(LogIoStoreDependencyViewer, Warning, TEXT("Failed to delete PartCache directory: %s"), *PartCacheDir);
+			}
+		}
+	}
+
 	// Reset cloud/partial download parameters to prevent stale state
 	bUsePartialDownloads = false;
 	ZenExePath.Empty();
diff --git a/Engine/Source/Programs/IoStoreDependencyViewer/Private/SIoStoreDependencyViewer.h b/Engine/Source/Programs/IoStoreDependencyViewer/Private/SIoStoreDependencyViewer.h
--- a/Engine/Source/Programs/IoStoreDependencyViewer/Private/SIoStoreDependencyViewer.h
+++ b/Engine/Source/Programs/IoStoreDependencyViewer/Private/SIoStoreDependencyViewer.h
@@ -75,7 +75,7 @@ private:
 	// Loading functions
 	void ClearAllData(); // Clears all data and UI state
 	bool LoadFromDirectory(const FString& DirectoryPath);
-	bool LoadStandardContainers(const FString& DirectoryPath);
+	bool LoadAllContainersWithProgress(const FString& DirectoryPath);
 	bool LoadOnDemandContainers(const FString& DirectoryPath);
 	bool LoadMetaData(const FString& DirectoryPath);
 	void ProcessContainerHeader(const FIoContainerHeader& Header, const FString& ContainerName);
@@ -192,6 +192,9 @@ private:
 	FString CloudProxyUrl;
 	FString BaseDownloadDirectory; // Where UTOCs are stored
 
+	// Current loaded build info (for status display)
+	FString LoadedBuildSource; // Path or cloud build info
+
 	static constexpr int32 MaxTreeDepth = 10;
 
 public:
diff --git a/Engine/Source/Programs/IoStoreDependencyViewer/Private/TestCloudDownload.cpp b/Engine/Source/Programs/IoStoreDependencyViewer/Private/TestCloudDownload.cpp
--- a/Engine/Source/Programs/IoStoreDependencyViewer/Private/TestCloudDownload.cpp
+++ b/Engine/Source/Programs/IoStoreDependencyViewer/Private/TestCloudDownload.cpp
@@ -572,7 +572,7 @@ bool TestLoadFromDirectory(const FString& DirectoryPath)
 	PlatformFile.FindFilesRecursively(MetaFiles, *DirectoryPath, TEXT(".umeta"));
 	UE_LOG(LogIoStoreDependencyViewer, Display, TEXT("Found %d .umeta files"), MetaFiles.Num());
 
-	// Step 2: Load .utoc files (same as SIoStoreDependencyViewer::LoadStandardContainers)
+	// Step 2: Load .utoc files (same as SIoStoreDependencyViewer::LoadAllContainersWithProgress)
 	TArray<FString> TocFiles;
 	PlatformFile.FindFilesRecursively(TocFiles, *DirectoryPath, TEXT(".utoc"));
 
